Graph-driven pose warping, with root motion delta flow:

1) Root motion delta flow supported through all anim nodes deriving from FAnimNode_AssetPlayerBase and montages.
2) Stride/Orientation warping feature graph-driven evaluation modes which leverage root motion delta flow (graph-driven slope warping currently disabled).
3) Root motion flow connected to the Animation Warping plugin.
4) Pose Warping test map featuring various motion "styles" connected to pose warping (Motion Matching, Blend Space Graphs/Assets, Sequence Players, Evaluators).

diff --git a/Engine/Plugins/Experimental/Animation/AnimationWarping/Source/Editor/Private/AnimGraph/AnimGraphNode_OrientationWarping.cpp b/Engine/Plugins/Experimental/Animation/AnimationWarping/Source/Editor/Private/AnimGraph/AnimGraphNode_OrientationWarping.cpp
--- a/Engine/Plugins/Experimental/Animation/AnimationWarping/Source/Editor/Private/AnimGraph/AnimGraphNode_OrientationWarping.cpp
+++ b/Engine/Plugins/Experimental/Animation/AnimationWarping/Source/Editor/Private/AnimGraph/AnimGraphNode_OrientationWarping.cpp
@@ -1,7 +1,9 @@
 // Copyright Epic Games, Inc. All Rights Reserved.
 
-
 #include "AnimGraph/AnimGraphNode_OrientationWarping.h"
+#include "Animation/AnimRootMotionProvider.h"
+#include "Kismet2/BlueprintEditorUtils.h"
+#include "DetailLayoutBuilder.h"
 
 /////////////////////////////////////////////////////
 // UAnimGraphNode_OrientationWarping
@@ -20,7 +22,7 @@ FText UAnimGraphNode_OrientationWarping::GetControllerDescription() const
 
 FText UAnimGraphNode_OrientationWarping::GetTooltipText() const
 {
-	return LOCTEXT("OrientationWarpingTooltip", "Orients RootBone to match locomotion direction, and counter rotates spine.");
+	return LOCTEXT("OrientationWarpingTooltip", "Orients Root Bone to match locomotion direction, and counter rotates spine.");
 }
 
 FText UAnimGraphNode_OrientationWarping::GetNodeTitle(ENodeTitleType::Type TitleType) const
@@ -28,4 +30,53 @@ FText UAnimGraphNode_OrientationWarping::GetNodeTitle(ENodeTitleType::Type Title
 	return GetControllerDescription();
 }
 
+void UAnimGraphNode_OrientationWarping::PostEditChangeProperty(struct FPropertyChangedEvent& PropertyChangedEvent)
+{
+	DECLARE_SCOPE_HIERARCHICAL_COUNTER_FUNC()
+	Super::PostEditChangeProperty(PropertyChangedEvent);
+
+	bool bRequiresNodeReconstruct = false;
+	FProperty* ChangedProperty = PropertyChangedEvent.Property;
+
+	if (ChangedProperty)
+	{
+		if ((ChangedProperty->GetFName() == GET_MEMBER_NAME_STRING_CHECKED(FInputScaleBiasClamp, bMapRange))
+			|| (ChangedProperty->GetFName() == GET_MEMBER_NAME_STRING_CHECKED(FInputRange, Min))
+			|| (ChangedProperty->GetFName() == GET_MEMBER_NAME_STRING_CHECKED(FInputRange, Max))
+			|| (ChangedProperty->GetFName() == GET_MEMBER_NAME_STRING_CHECKED(FInputScaleBiasClamp, Scale))
+			|| (ChangedProperty->GetFName() == GET_MEMBER_NAME_STRING_CHECKED(FInputScaleBiasClamp, Bias))
+			|| (ChangedProperty->GetFName() == GET_MEMBER_NAME_STRING_CHECKED(FInputScaleBiasClamp, bClampResult))
+			|| (ChangedProperty->GetFName() == GET_MEMBER_NAME_STRING_CHECKED(FInputScaleBiasClamp, ClampMin))
+			|| (ChangedProperty->GetFName() == GET_MEMBER_NAME_STRING_CHECKED(FInputScaleBiasClamp, ClampMax))
+			|| (ChangedProperty->GetFName() == GET_MEMBER_NAME_STRING_CHECKED(FInputScaleBiasClamp, bInterpResult))
+			|| (ChangedProperty->GetFName() == GET_MEMBER_NAME_STRING_CHECKED(FInputScaleBiasClamp, InterpSpeedIncreasing))
+			|| (ChangedProperty->GetFName() == GET_MEMBER_NAME_STRING_CHECKED(FInputScaleBiasClamp, InterpSpeedDecreasing)))
+		{
+			bRequiresNodeReconstruct = true;
+		}
+	}
+
+	if (bRequiresNodeReconstruct)
+	{
+		ReconstructNode();
+		FBlueprintEditorUtils::MarkBlueprintAsStructurallyModified(GetBlueprint());
+	}
+}
+
+void UAnimGraphNode_OrientationWarping::GetInputLinkAttributes(FNodeAttributeArray& OutAttributes) const
+{
+	if (Node.Mode == EWarpingEvaluationMode::Graph)
+	{
+		OutAttributes.Add(UE::Anim::IAnimRootMotionProvider::RootMotionDeltaAttributeName);
+	}
+}
+
+void UAnimGraphNode_OrientationWarping::GetOutputLinkAttributes(FNodeAttributeArray& OutAttributes) const
+{
+	if (Node.Mode == EWarpingEvaluationMode::Graph)
+	{
+		OutAttributes.Add(UE::Anim::IAnimRootMotionProvider::RootMotionDeltaAttributeName);
+	}
+}
+
 #undef LOCTEXT_NAMESPACE
diff --git a/Engine/Plugins/Experimental/Animation/AnimationWarping/Source/Editor/Private/AnimGraph/AnimGraphNode_SlopeWarping.cpp b/Engine/Plugins/Experimental/Animation/AnimationWarping/Source/Editor/Private/AnimGraph/AnimGraphNode_SlopeWarping.cpp
--- a/Engine/Plugins/Experimental/Animation/AnimationWarping/Source/Editor/Private/AnimGraph/AnimGraphNode_SlopeWarping.cpp
+++ b/Engine/Plugins/Experimental/Animation/AnimationWarping/Source/Editor/Private/AnimGraph/AnimGraphNode_SlopeWarping.cpp
@@ -2,6 +2,7 @@
 
 
 #include "AnimGraph/AnimGraphNode_SlopeWarping.h"
+#include "Animation/AnimRootMotionProvider.h"
 
 #define LOCTEXT_NAMESPACE "MomentumNodes"
 
@@ -24,4 +25,20 @@ FText UAnimGraphNode_SlopeWarping::GetNodeTitle(ENodeTitleType::Type TitleType)
 	return GetControllerDescription();
 }
 
+void UAnimGraphNode_SlopeWarping::GetInputLinkAttributes(FNodeAttributeArray& OutAttributes) const
+{
+	if (Node.Mode == EWarpingEvaluationMode::Graph)
+	{
+		OutAttributes.Add(UE::Anim::IAnimRootMotionProvider::RootMotionDeltaAttributeName);
+	}
+}
+
+void UAnimGraphNode_SlopeWarping::GetOutputLinkAttributes(FNodeAttributeArray& OutAttributes) const
+{
+	if (Node.Mode == EWarpingEvaluationMode::Graph)
+	{
+		OutAttributes.Add(UE::Anim::IAnimRootMotionProvider::RootMotionDeltaAttributeName);
+	}
+}
+
 #undef LOCTEXT_NAMESPACE
diff --git a/Engine/Plugins/Experimental/Animation/AnimationWarping/Source/Editor/Private/AnimGraph/AnimGraphNode_StrideWarping.cpp b/Engine/Plugins/Experimental/Animation/AnimationWarping/Source/Editor/Private/AnimGraph/AnimGraphNode_StrideWarping.cpp
--- a/Engine/Plugins/Experimental/Animation/AnimationWarping/Source/Editor/Private/AnimGraph/AnimGraphNode_StrideWarping.cpp
+++ b/Engine/Plugins/Experimental/Animation/AnimationWarping/Source/Editor/Private/AnimGraph/AnimGraphNode_StrideWarping.cpp
@@ -1,7 +1,9 @@
 // Copyright Epic Games, Inc. All Rights Reserved.
 
-
 #include "AnimGraph/AnimGraphNode_StrideWarping.h"
+#include "Animation/AnimRootMotionProvider.h"
+#include "Kismet2/BlueprintEditorUtils.h"
+#include "DetailLayoutBuilder.h"
 
 #define LOCTEXT_NAMESPACE "MomentumNodes"
 
@@ -28,36 +30,128 @@ void UAnimGraphNode_StrideWarping::CustomizePinData(UEdGraphPin* Pin, FName Sour
 {
 	Super::CustomizePinData(Pin, SourcePropertyName, ArrayIndex);
 
-	if (Pin->PinName == GET_MEMBER_NAME_STRING_CHECKED(FAnimNode_StrideWarping, StrideScaling))
+	if (Pin->PinName == GET_MEMBER_NAME_STRING_CHECKED(FAnimNode_StrideWarping, ManualStrideWarpingDir))
+	{
+		Pin->bHidden = (Node.Mode == EWarpingEvaluationMode::Graph);
+	}
+
+	if (Pin->PinName == GET_MEMBER_NAME_STRING_CHECKED(FAnimNode_StrideWarping, ManualStrideScaling))
 	{
+		Pin->bHidden = (Node.Mode == EWarpingEvaluationMode::Graph);
+
 		if (!Pin->bHidden)
 		{
 			Pin->PinFriendlyName = Node.StrideScalingScaleBiasClamp.GetFriendlyName(Pin->PinFriendlyName);
 		}
 	}
+
+	if (Pin->PinName == GET_MEMBER_NAME_STRING_CHECKED(FAnimNode_StrideWarping, LocomotionSpeed))
+	{
+		Pin->bHidden = (Node.Mode == EWarpingEvaluationMode::Manual);
+	}
+}
+
+void UAnimGraphNode_StrideWarping::CustomizeDetails(IDetailLayoutBuilder& DetailBuilder)
+{
+	DECLARE_SCOPE_HIERARCHICAL_COUNTER_FUNC()
+	Super::CustomizeDetails(DetailBuilder);
+
+	TSharedRef<IPropertyHandle> NodeHandle = DetailBuilder.GetProperty(FName(TEXT("Node")), GetClass());
+
+	if (Node.Mode == EWarpingEvaluationMode::Graph)
+	{
+		DetailBuilder.HideProperty(NodeHandle->GetChildHandle(GET_MEMBER_NAME_CHECKED(FAnimNode_StrideWarping, ManualStrideWarpingDir)));
+		DetailBuilder.HideProperty(NodeHandle->GetChildHandle(GET_MEMBER_NAME_CHECKED(FAnimNode_StrideWarping, ManualStrideScaling)));
+	}
+
+	if (Node.Mode == EWarpingEvaluationMode::Manual)
+	{
+		DetailBuilder.HideProperty(NodeHandle->GetChildHandle(GET_MEMBER_NAME_CHECKED(FAnimNode_StrideWarping, LocomotionSpeed)));
+	}
 }
 
 void UAnimGraphNode_StrideWarping::PostEditChangeProperty(struct FPropertyChangedEvent& PropertyChangedEvent)
 {
-	const FName PropertyName = (PropertyChangedEvent.Property ? PropertyChangedEvent.Property->GetFName() : NAME_None);
-
-	// Reconstruct node to show updates to PinFriendlyNames.
-	if ((PropertyName == GET_MEMBER_NAME_STRING_CHECKED(FInputScaleBiasClamp, bMapRange))
-		|| (PropertyName == GET_MEMBER_NAME_STRING_CHECKED(FInputRange, Min))
-		|| (PropertyName == GET_MEMBER_NAME_STRING_CHECKED(FInputRange, Max))
-		|| (PropertyName == GET_MEMBER_NAME_STRING_CHECKED(FInputScaleBiasClamp, Scale))
-		|| (PropertyName == GET_MEMBER_NAME_STRING_CHECKED(FInputScaleBiasClamp, Bias))
-		|| (PropertyName == GET_MEMBER_NAME_STRING_CHECKED(FInputScaleBiasClamp, bClampResult))
-		|| (PropertyName == GET_MEMBER_NAME_STRING_CHECKED(FInputScaleBiasClamp, ClampMin))
-		|| (PropertyName == GET_MEMBER_NAME_STRING_CHECKED(FInputScaleBiasClamp, ClampMax))
-		|| (PropertyName == GET_MEMBER_NAME_STRING_CHECKED(FInputScaleBiasClamp, bInterpResult))
-		|| (PropertyName == GET_MEMBER_NAME_STRING_CHECKED(FInputScaleBiasClamp, InterpSpeedIncreasing))
-		|| (PropertyName == GET_MEMBER_NAME_STRING_CHECKED(FInputScaleBiasClamp, InterpSpeedDecreasing)))
+	DECLARE_SCOPE_HIERARCHICAL_COUNTER_FUNC()
+	Super::PostEditChangeProperty(PropertyChangedEvent);
+
+	bool bRequiresNodeReconstruct = false;
+	FProperty* ChangedProperty = PropertyChangedEvent.Property;
+
+	if (ChangedProperty)
+	{
+		if ((ChangedProperty->GetFName() == GET_MEMBER_NAME_STRING_CHECKED(FInputScaleBiasClamp, bMapRange))
+			|| (ChangedProperty->GetFName() == GET_MEMBER_NAME_STRING_CHECKED(FInputRange, Min))
+			|| (ChangedProperty->GetFName() == GET_MEMBER_NAME_STRING_CHECKED(FInputRange, Max))
+			|| (ChangedProperty->GetFName() == GET_MEMBER_NAME_STRING_CHECKED(FInputScaleBiasClamp, Scale))
+			|| (ChangedProperty->GetFName() == GET_MEMBER_NAME_STRING_CHECKED(FInputScaleBiasClamp, Bias))
+			|| (ChangedProperty->GetFName() == GET_MEMBER_NAME_STRING_CHECKED(FInputScaleBiasClamp, bClampResult))
+			|| (ChangedProperty->GetFName() == GET_MEMBER_NAME_STRING_CHECKED(FInputScaleBiasClamp, ClampMin))
+			|| (ChangedProperty->GetFName() == GET_MEMBER_NAME_STRING_CHECKED(FInputScaleBiasClamp, ClampMax))
+			|| (ChangedProperty->GetFName() == GET_MEMBER_NAME_STRING_CHECKED(FInputScaleBiasClamp, bInterpResult))
+			|| (ChangedProperty->GetFName() == GET_MEMBER_NAME_STRING_CHECKED(FInputScaleBiasClamp, InterpSpeedIncreasing))
+			|| (ChangedProperty->GetFName() == GET_MEMBER_NAME_STRING_CHECKED(FInputScaleBiasClamp, InterpSpeedDecreasing)))
+		{
+			bRequiresNodeReconstruct = true;
+		}
+
+		if (ChangedProperty->GetFName() == GET_MEMBER_NAME_STRING_CHECKED(FAnimNode_StrideWarping, Mode))
+		{
+			FScopedTransaction Transaction(LOCTEXT("ChangeEvaluationMode", "Change Evaluation Mode"));
+			Modify();
+
+			// Break links to pins going away
+			for (int32 PinIndex = 0; PinIndex < Pins.Num(); ++PinIndex)
+			{
+				UEdGraphPin* Pin = Pins[PinIndex];
+				if (Pin->PinName == GET_MEMBER_NAME_STRING_CHECKED(FAnimNode_StrideWarping, ManualStrideWarpingDir))
+				{
+					if (Node.Mode == EWarpingEvaluationMode::Graph)
+					{
+						Pin->BreakAllPinLinks();
+					}
+				}
+				else if (Pin->PinName == GET_MEMBER_NAME_STRING_CHECKED(FAnimNode_StrideWarping, ManualStrideScaling))
+				{
+					if (Node.Mode == EWarpingEvaluationMode::Graph)
+					{
+						Pin->BreakAllPinLinks();
+					}
+				}
+				else if (Pin->PinName == GET_MEMBER_NAME_STRING_CHECKED(FAnimNode_StrideWarping, LocomotionSpeed))
+				{
+					if (Node.Mode == EWarpingEvaluationMode::Manual)
+					{
+						Pin->BreakAllPinLinks();
+					}
+				}
+			}
+
+			bRequiresNodeReconstruct = true;
+		}
+	}
+
+	if (bRequiresNodeReconstruct)
 	{
 		ReconstructNode();
+		FBlueprintEditorUtils::MarkBlueprintAsStructurallyModified(GetBlueprint());
 	}
+}
 
-	Super::PostEditChangeProperty(PropertyChangedEvent);
+void UAnimGraphNode_StrideWarping::GetInputLinkAttributes(FNodeAttributeArray& OutAttributes) const
+{
+	if (Node.Mode == EWarpingEvaluationMode::Graph)
+	{
+		OutAttributes.Add(UE::Anim::IAnimRootMotionProvider::RootMotionDeltaAttributeName);
+	}
+}
+
+void UAnimGraphNode_StrideWarping::GetOutputLinkAttributes(FNodeAttributeArray& OutAttributes) const
+{
+	if (Node.Mode == EWarpingEvaluationMode::Graph)
+	{
+		OutAttributes.Add(UE::Anim::IAnimRootMotionProvider::RootMotionDeltaAttributeName);
+	}
 }
 
 #undef LOCTEXT_NAMESPACE
diff --git a/Engine/Plugins/Experimental/Animation/AnimationWarping/Source/Editor/Public/AnimGraph/AnimGraphNode_OrientationWarping.h b/Engine/Plugins/Experimental/Animation/AnimationWarping/Source/Editor/Public/AnimGraph/AnimGraphNode_OrientationWarping.h
--- a/Engine/Plugins/Experimental/Animation/AnimationWarping/Source/Editor/Public/AnimGraph/AnimGraphNode_OrientationWarping.h
+++ b/Engine/Plugins/Experimental/Animation/AnimationWarping/Source/Editor/Public/AnimGraph/AnimGraphNode_OrientationWarping.h
@@ -17,6 +17,9 @@ public:
 	// UEdGraphNode interface
 	virtual FText GetNodeTitle(ENodeTitleType::Type TitleType) const override;
 	virtual FText GetTooltipText() const override;
+	virtual void PostEditChangeProperty(struct FPropertyChangedEvent& PropertyChangedEvent) override;
+	virtual void GetInputLinkAttributes(FNodeAttributeArray& OutAttributes) const override;
+	virtual void GetOutputLinkAttributes(FNodeAttributeArray& OutAttributes) const override;
 	// End of UEdGraphNode interface
 
 protected:
diff --git a/Engine/Plugins/Experimental/Animation/AnimationWarping/Source/Editor/Public/AnimGraph/AnimGraphNode_SlopeWarping.h b/Engine/Plugins/Experimental/Animation/AnimationWarping/Source/Editor/Public/AnimGraph/AnimGraphNode_SlopeWarping.h
--- a/Engine/Plugins/Experimental/Animation/AnimationWarping/Source/Editor/Public/AnimGraph/AnimGraphNode_SlopeWarping.h
+++ b/Engine/Plugins/Experimental/Animation/AnimationWarping/Source/Editor/Public/AnimGraph/AnimGraphNode_SlopeWarping.h
@@ -19,6 +19,8 @@ public:
 	// UEdGraphNode interface
 	virtual FText GetNodeTitle(ENodeTitleType::Type TitleType) const override;
 	virtual FText GetTooltipText() const override;
+	virtual void GetInputLinkAttributes(FNodeAttributeArray& OutAttributes) const override;
+	virtual void GetOutputLinkAttributes(FNodeAttributeArray& OutAttributes) const override;
 	// End of UEdGraphNode interface
 
 protected:
diff --git a/Engine/Plugins/Experimental/Animation/AnimationWarping/Source/Editor/Public/AnimGraph/AnimGraphNode_StrideWarping.h b/Engine/Plugins/Experimental/Animation/AnimationWarping/Source/Editor/Public/AnimGraph/AnimGraphNode_StrideWarping.h
--- a/Engine/Plugins/Experimental/Animation/AnimationWarping/Source/Editor/Public/AnimGraph/AnimGraphNode_StrideWarping.h
+++ b/Engine/Plugins/Experimental/Animation/AnimationWarping/Source/Editor/Public/AnimGraph/AnimGraphNode_StrideWarping.h
@@ -19,7 +19,10 @@ public:
 	// UEdGraphNode interface
 	virtual FText GetNodeTitle(ENodeTitleType::Type TitleType) const override;
 	virtual FText GetTooltipText() const override;
+	virtual void CustomizeDetails(IDetailLayoutBuilder& DetailBuilder) override;
 	virtual void PostEditChangeProperty(struct FPropertyChangedEvent& PropertyChangedEvent) override;
+	virtual void GetInputLinkAttributes(FNodeAttributeArray& OutAttributes) const override;
+	virtual void GetOutputLinkAttributes(FNodeAttributeArray& OutAttributes) const override;
 	// End of UEdGraphNode interface
 
 protected:
diff --git a/Engine/Plugins/Experimental/Animation/AnimationWarping/Source/Runtime/Private/AnimationWarpingRuntimeModule.cpp b/Engine/Plugins/Experimental/Animation/AnimationWarping/Source/Runtime/Private/AnimationWarpingRuntimeModule.cpp
--- a/Engine/Plugins/Experimental/Animation/AnimationWarping/Source/Runtime/Private/AnimationWarpingRuntimeModule.cpp
+++ b/Engine/Plugins/Experimental/Animation/AnimationWarping/Source/Runtime/Private/AnimationWarpingRuntimeModule.cpp
@@ -1,11 +1,73 @@
 // Copyright Epic Games, Inc. All Rights Reserved.
 
 #include "AnimationWarpingRuntimeModule.h"
+#include "Animation/AnimRootMotionProvider.h"
+#include "Animation/AnimNodeBase.h"
+#include "Animation/AnimSequence.h"
+#include "Animation/BuiltInAttributeTypes.h"
 #include "Modules/ModuleManager.h"
 
 class FAnimationWarpingRuntimeModule : public IAnimationWarpingRuntimeModule
 {
 };
 
+namespace UE { namespace AnimationWarping {
 
-IMPLEMENT_MODULE( FAnimationWarpingRuntimeModule, AnimationWarpingRuntime )
+// Root motion animation attribute data will always be accessible on the root bone's attribute container
+static const UE::Anim::FAttributeId RootMotionAttributeId = { UE::Anim::IAnimRootMotionProvider::RootMotionDeltaAttributeName , FCompactPoseBoneIndex(0) };
+
+class FModule : public IModuleInterface, public UE::Anim::IAnimRootMotionProvider
+{
+public:
+	virtual void StartupModule() override;
+	virtual void ShutdownModule() override;
+
+public:
+	virtual void SampleRootMotion(const FDeltaTimeRecord& SampleRange, const UAnimSequence& Sequence, bool bLoopingSequence, UE::Anim::FStackAttributeContainer& OutAttributes) const override;
+	virtual bool OverrideRootMotion(const FTransform& RootMotionDelta, UE::Anim::FStackAttributeContainer& OutAttributes) const override;
+	virtual bool ExtractRootMotion(const UE::Anim::FStackAttributeContainer& Attributes, FTransform& OutRootMotionDelta) const override;
+	virtual bool HasRootMotion(const UE::Anim::FStackAttributeContainer& Attributes) const override;
+};
+
+void FModule::StartupModule()
+{
+	IModularFeatures::Get().RegisterModularFeature(UE::Anim::IAnimRootMotionProvider::ModularFeatureName, this);
+}
+
+void FModule::ShutdownModule()
+{
+	IModularFeatures::Get().UnregisterModularFeature(UE::Anim::IAnimRootMotionProvider::ModularFeatureName, this);
+}
+
+void FModule::SampleRootMotion(const FDeltaTimeRecord& SampleRange, const UAnimSequence& Sequence, bool bLoopingSequence, UE::Anim::FStackAttributeContainer& OutAttributes) const
+{
+	const FTransform RootMotionTransform = Sequence.ExtractRootMotion(SampleRange.Previous, SampleRange.Delta, bLoopingSequence);
+	OutAttributes.Add<FTransformAnimationAttribute>(RootMotionAttributeId, { RootMotionTransform });
+}
+
+bool FModule::OverrideRootMotion(const FTransform& RootMotionDelta, UE::Anim::FStackAttributeContainer& OutAttributes) const
+{
+	if (FTransformAnimationAttribute* RootMotionAttribute = OutAttributes.Find<FTransformAnimationAttribute>(RootMotionAttributeId))
+	{
+		RootMotionAttribute->Value = RootMotionDelta;
+		return true;
+	}
+
+	return false;
+}
+
+bool FModule::ExtractRootMotion(const UE::Anim::FStackAttributeContainer& Attributes, FTransform& OutRootMotionDelta) const
+{
+	const FTransformAnimationAttribute* RootMotionAttribute = Attributes.Find<FTransformAnimationAttribute>(RootMotionAttributeId);
+	OutRootMotionDelta = RootMotionAttribute ? RootMotionAttribute->Value : FTransform::Identity;
+	return !!RootMotionAttribute;
+}
+
+bool FModule::HasRootMotion(const UE::Anim::FStackAttributeContainer& Attributes) const
+{
+	return !!Attributes.Find<FTransformAnimationAttribute>(RootMotionAttributeId);
+}
+
+}}
+
+IMPLEMENT_MODULE(UE::AnimationWarping::FModule, AnimationWarpingRuntime)
diff --git a/Engine/Plugins/Experimental/Animation/AnimationWarping/Source/Runtime/Private/BoneControllers/AnimNode_OrientationWarping.cpp b/Engine/Plugins/Experimental/Animation/AnimationWarping/Source/Runtime/Private/BoneControllers/AnimNode_OrientationWarping.cpp
--- a/Engine/Plugins/Experimental/Animation/AnimationWarping/Source/Runtime/Private/BoneControllers/AnimNode_OrientationWarping.cpp
+++ b/Engine/Plugins/Experimental/Animation/AnimationWarping/Source/Runtime/Private/BoneControllers/AnimNode_OrientationWarping.cpp
@@ -1,19 +1,20 @@
 // Copyright Epic Games, Inc. All Rights Reserved.
 
 #include "BoneControllers/AnimNode_OrientationWarping.h"
+#include "Animation/AnimRootMotionProvider.h"
 #include "AnimationRuntime.h"
 #include "Animation/AnimInstanceProxy.h"
 
 DECLARE_CYCLE_STAT(TEXT("OrientationWarping Eval"), STAT_OrientationWarping_Eval, STATGROUP_Anim);
 
 FAnimNode_OrientationWarping::FAnimNode_OrientationWarping()
-	: LocomotionAngle(0.f)
+	: Mode(EWarpingEvaluationMode::Manual)
+	, LocomotionAngle(0.f)
 	, IKFootRootBoneIndex(INDEX_NONE)
-	, CachedDeltaTime(0.f)
 {
 }
 
-static FVector GetAxisVector(const EAxis::Type& InAxis)
+static inline FVector GetAxisVector(const EAxis::Type& InAxis)
 {
 	switch (InAxis)
 	{
@@ -24,6 +25,11 @@ static FVector GetAxisVector(const EAxis::Type& InAxis)
 	};
 }
 
+static inline bool IsInvalidWarpingAngle(float Angle, float Tolerance)
+{
+	return FMath::IsNearlyZero(Angle, Tolerance) || FMath::IsNearlyEqual(FMath::Abs(Angle), PI, Tolerance);
+}
+
 void FAnimNode_OrientationWarping::GatherDebugData(FNodeDebugData& DebugData)
 {
 	FString DebugLine = DebugData.GetNodeName(this);
@@ -42,8 +48,6 @@ void FAnimNode_OrientationWarping::Initialize_AnyThread(const FAnimationInitiali
 void FAnimNode_OrientationWarping::UpdateInternal(const FAnimationUpdateContext& Context)
 {
 	FAnimNode_SkeletalControlBase::UpdateInternal(Context);
-
-	CachedDeltaTime += Context.GetDeltaTime();
 }
 
 void FAnimNode_OrientationWarping::EvaluateSkeletalControl_AnyThread(FComponentSpacePoseContext& Output, TArray<FBoneTransform>& OutBoneTransforms)
@@ -51,86 +55,144 @@ void FAnimNode_OrientationWarping::EvaluateSkeletalControl_AnyThread(FComponentS
 	SCOPE_CYCLE_COUNTER(STAT_OrientationWarping_Eval);
 
 	check(OutBoneTransforms.Num() == 0);
-	const FBoneContainer& RequiredBones = Output.Pose.GetPose().GetBoneContainer();
 
-	const float YawAngleRadians = FMath::DegreesToRadians(FRotator::NormalizeAxis(LocomotionAngle));
-	if (!FMath::IsNearlyZero(YawAngleRadians, KINDA_SMALL_NUMBER))
+	// We will likely need to revisit LocomotionAngle participating as an input to orientation warping.
+	// Without velocity information from the motion model (such as the capsule), LocomotionAngle isn't enough
+	// information in isolation for all cases when deciding to warp.
+	//
+	// For example imagine that the motion model has stopped moving with zero velocity due to a
+	// transition into a strafing stop. During that transition we may play an animation with non-zero 
+	// velocity for an arbitrary number of frames. In this scenario the concept of direction is meaningless 
+	// since we cannot orient the animation to match a zero velocity and consequently a zero direction, 
+	// since that would break the pose. For those frames, we would incorrectly over-orient the strafe.
+	//
+	// The solution may be instead to pass velocity with the actor base rotation, allowing us to retain
+	// speed information about the motion. It may also allow us to do more complex orienting behavior 
+	// when multiple degrees of freedom can be considered.
+	float LocomotionRotationAngle = FMath::DegreesToRadians(FRotator::NormalizeAxis(LocomotionAngle));
+
+	const FVector LocomotionRotationAxis = GetAxisVector(Settings.YawRotationAxis);
+	const UE::Anim::IAnimRootMotionProvider* RootMotionProvider = UE::Anim::IAnimRootMotionProvider::Get();
+	const bool bGraphDrivenWarping = RootMotionProvider && Mode == EWarpingEvaluationMode::Graph;
+
+	// Graph driven orientation warping will modify the incoming root motion to orient towards the intended locomotion angle
+	if (bGraphDrivenWarping)
 	{
-		const FVector RotationAxis = GetAxisVector(Settings.YawRotationAxis);
-		const float BodyOrientationAlpha = FMath::Clamp(Settings.BodyOrientationAlpha, 0.f, 1.f);
-		const float IKFootRootOrientationAlpha = 1.f - BodyOrientationAlpha;
+		if (IsInvalidWarpingAngle(LocomotionRotationAngle, KINDA_SMALL_NUMBER))
+		{
+			return;
+		}
 
-		// Rotate Root Bone first, as that cheaply rotates the whole pose with one transformation.
-		if (!FMath::IsNearlyZero(BodyOrientationAlpha, KINDA_SMALL_NUMBER))
+		FTransform RootMotionTransformDelta;
+		const bool bRootMotionDeltaPresent = RootMotionProvider->ExtractRootMotion(Output.CustomAttributes, RootMotionTransformDelta);
+
+		if (ensure(bRootMotionDeltaPresent))
 		{
-			const FQuat RootRotation = FQuat(RotationAxis, YawAngleRadians * BodyOrientationAlpha);
+			// In UE, forward is defined as +x; consequently this is also true when sampling an actor's velocity. Sometimes the skeletal 
+			// mesh component forward will not match the actor, requiring us to correct the rotation before sampling the LocomotionForward.
+			// In order to make orientation warping 'pure' in the future we will need to provide more context about the intent of
+			// the actor vs the intent of the animation in their respective spaces. Specifically, we will need some form the following information:
+			//
+			// 1. Actor Forward
+			// 2. Actor Velocity
+			// 3. Skeletal Mesh Relative Rotation
+			const FTransform SkeletalMeshRelativeTransform = Output.AnimInstanceProxy->GetComponentRelativeTransform();
+			const FQuat SkeletalMeshRelativeRotation = SkeletalMeshRelativeTransform.GetRotation();
+			const FQuat LocomotionRotation = FQuat(LocomotionRotationAxis, LocomotionRotationAngle);
+			const FVector LocomotionForwardDir = SkeletalMeshRelativeRotation.UnrotateVector(LocomotionRotation.GetForwardVector());
+			const FVector RootMotionDeltaDir = RootMotionTransformDelta.GetTranslation();
+
+			// Capture the delta rotation from the axis of motion we care about
+			const FQuat WarpedRotation = FQuat::FindBetween(RootMotionDeltaDir, LocomotionForwardDir);
+			LocomotionRotationAngle = WarpedRotation.GetTwistAngle(LocomotionRotationAxis);
+
+			if (IsInvalidWarpingAngle(LocomotionRotationAngle, KINDA_SMALL_NUMBER))
+			{
+				return;
+			}
+
+			// Rotate the root motion delta fully by the warped angle
+			const FVector RootMotionTranslationDelta = RootMotionTransformDelta.GetTranslation();
+			const FVector WarpedRootMotionTranslationDelta = WarpedRotation.RotateVector(RootMotionTranslationDelta);
+			RootMotionTransformDelta.SetTranslation(WarpedRootMotionTranslationDelta);
 
-			const FCompactPoseBoneIndex RootBoneIndex(0);
-			FTransform RootBoneTransform(Output.Pose.GetComponentSpaceTransform(RootBoneIndex));
-			RootBoneTransform.SetRotation(RootRotation * RootBoneTransform.GetRotation());
-			RootBoneTransform.NormalizeRotation();
-			Output.Pose.SetComponentSpaceTransform(RootBoneIndex, RootBoneTransform);
+			// Forward the side effects of orientation warping on the root motion contribution for this sub-graph
+			const bool bRootMotionOverridden = RootMotionProvider->OverrideRootMotion(RootMotionTransformDelta, Output.CustomAttributes);
+			ensure(bRootMotionOverridden);
 		}
-		
+	} 
+	else if (IsInvalidWarpingAngle(LocomotionRotationAngle, KINDA_SMALL_NUMBER))
+	{
+		return;
+	}
+
+	const float BodyOrientationAlpha = FMath::Clamp(Settings.BodyOrientationAlpha, 0.f, 1.f);
+
+	// Rotate Root Bone first, as that cheaply rotates the whole pose with one transformation.
+	if (!FMath::IsNearlyZero(BodyOrientationAlpha, KINDA_SMALL_NUMBER))
+	{
+		const FQuat RootRotation = FQuat(LocomotionRotationAxis, LocomotionRotationAngle * BodyOrientationAlpha);
+		const FCompactPoseBoneIndex RootBoneIndex(0);
+
+		FTransform RootBoneTransform(Output.Pose.GetComponentSpaceTransform(RootBoneIndex));
+		RootBoneTransform.SetRotation(RootRotation * RootBoneTransform.GetRotation());
+		RootBoneTransform.NormalizeRotation();
+		Output.Pose.SetComponentSpaceTransform(RootBoneIndex, RootBoneTransform);
+
+		const float IKFootRootOrientationAlpha = 1.f - BodyOrientationAlpha;
 		const int32 NumSpineBones = SpineBoneDataArray.Num();
 		const int32 NumIKFootBones = IKFootBoneIndexArray.Num();
-
-		const bool bUpdateSpineBones = (NumSpineBones > 0) && !FMath::IsNearlyZero(BodyOrientationAlpha, KINDA_SMALL_NUMBER);
+		const bool bBodyOrientationAlpha = !FMath::IsNearlyZero(BodyOrientationAlpha, KINDA_SMALL_NUMBER);
+		const bool bUpdateSpineBones = (NumSpineBones > 0) && bBodyOrientationAlpha;
 		const bool bUpdateIKFootRoot = (IKFootRootBoneIndex != FCompactPoseBoneIndex(INDEX_NONE)) && !FMath::IsNearlyZero(IKFootRootOrientationAlpha, KINDA_SMALL_NUMBER);
 		const bool bUpdateIKFootBones = bUpdateIKFootRoot && (NumIKFootBones > 0);
 
-		if (bUpdateSpineBones || bUpdateIKFootRoot || bUpdateIKFootBones)
+		if (bUpdateSpineBones)
 		{
-			if (bUpdateSpineBones)
+			// Spine bones counter rotate body orientation evenly across all bones.
+			for (int32 ArrayIndex = 0; ArrayIndex < NumSpineBones; ArrayIndex++)
 			{
-				// Spine bones counter rotate body orientation evenly across all bones.
-				for (int32 ArrayIndex = 0; ArrayIndex < NumSpineBones; ArrayIndex++)
-				{
-					const FOrientationWarpingSpineBoneData& BoneData = SpineBoneDataArray[ArrayIndex];
+				const FOrientationWarpingSpineBoneData& BoneData = SpineBoneDataArray[ArrayIndex];
+				const FQuat SpineBoneCounterRotation = FQuat(LocomotionRotationAxis, -LocomotionRotationAngle * BodyOrientationAlpha * BoneData.Weight);
+				check(BoneData.Weight > 0.f);
+
+				FTransform SpineBoneTransform(Output.Pose.GetComponentSpaceTransform(BoneData.BoneIndex));
+				SpineBoneTransform.SetRotation((SpineBoneCounterRotation * SpineBoneTransform.GetRotation()));
+				SpineBoneTransform.NormalizeRotation();
+				Output.Pose.SetComponentSpaceTransform(BoneData.BoneIndex, SpineBoneTransform);
+			}
+		}
 
-					const FQuat SpineBoneCounterRotation = FQuat(RotationAxis, -YawAngleRadians * BodyOrientationAlpha * BoneData.Weight);
-					check(BoneData.Weight > 0.f);
+		// Rotate IK Foot Root
+		if (bUpdateIKFootRoot)
+		{
+			const FQuat BoneRotation = FQuat(LocomotionRotationAxis, LocomotionRotationAngle * IKFootRootOrientationAlpha);
 
-					//FTransform SpineBoneTransform(Output.Pose.GetLocalSpaceTransform(BoneData.BoneIndex));
-					FTransform SpineBoneTransform(Output.Pose.GetComponentSpaceTransform(BoneData.BoneIndex));
-					SpineBoneTransform.SetRotation(SpineBoneCounterRotation * SpineBoneTransform.GetRotation());
-					SpineBoneTransform.NormalizeRotation();
-					Output.Pose.SetComponentSpaceTransform(BoneData.BoneIndex, SpineBoneTransform);
-				}
-			}
+			FTransform IKFootRootTransform(Output.Pose.GetComponentSpaceTransform(IKFootRootBoneIndex));
+			IKFootRootTransform.SetRotation(BoneRotation * IKFootRootTransform.GetRotation());
+			IKFootRootTransform.NormalizeRotation();
+			Output.Pose.SetComponentSpaceTransform(IKFootRootBoneIndex, IKFootRootTransform);
 
-			// Rotate IK Foot Root
-			if (bUpdateIKFootRoot)
+			// IK Feet 
+			// These match the root orientation, so don't rotate them. Just preserve root rotation. 
+			// We need to update their translation though, since we rotated their parent (the IK Foot Root bone).
+			if (bUpdateIKFootBones)
 			{
-				const FQuat BoneRotation = FQuat(RotationAxis, YawAngleRadians * IKFootRootOrientationAlpha);
-
-				FTransform IKFootRootTransform(Output.Pose.GetComponentSpaceTransform(IKFootRootBoneIndex));
-				IKFootRootTransform.SetRotation(BoneRotation * IKFootRootTransform.GetRotation());
-				IKFootRootTransform.NormalizeRotation();
-				Output.Pose.SetComponentSpaceTransform(IKFootRootBoneIndex, IKFootRootTransform);
+				const FQuat IKFootRotation = FQuat(LocomotionRotationAxis, -LocomotionRotationAngle * IKFootRootOrientationAlpha);
 
-				// IK Feet 
-				// These match the root orientation, so don't rotate them. Just preserve root rotation. 
-				// We need to update their translation though, since we rotated their parent (the IK Foot Root bone).
-				if (bUpdateIKFootBones)
+				for (int32 ArrayIndex = 0; ArrayIndex < NumIKFootBones; ArrayIndex++)
 				{
-					const FQuat IKFootRotation = FQuat(RotationAxis, -YawAngleRadians * IKFootRootOrientationAlpha);
-
-					for (int32 ArrayIndex = 0; ArrayIndex < NumIKFootBones; ArrayIndex++)
-					{
-						const FCompactPoseBoneIndex& IKFootBoneIndex = IKFootBoneIndexArray[ArrayIndex];
+					const FCompactPoseBoneIndex& IKFootBoneIndex = IKFootBoneIndexArray[ArrayIndex];
 
-						FTransform IKFootBoneTransform(Output.Pose.GetComponentSpaceTransform(IKFootBoneIndex));
-						IKFootBoneTransform.SetRotation(IKFootRotation * IKFootBoneTransform.GetRotation());
-						IKFootBoneTransform.NormalizeRotation();
-						Output.Pose.SetComponentSpaceTransform(IKFootBoneIndex, IKFootBoneTransform);
-					}
+					FTransform IKFootBoneTransform(Output.Pose.GetComponentSpaceTransform(IKFootBoneIndex));
+					IKFootBoneTransform.SetRotation(IKFootRotation * IKFootBoneTransform.GetRotation());
+					IKFootBoneTransform.NormalizeRotation();
+					Output.Pose.SetComponentSpaceTransform(IKFootBoneIndex, IKFootBoneTransform);
 				}
 			}
-
-			// Clear time accumulator, to be filled during next update.
-			CachedDeltaTime = 0.f;
 		}
+
+		OutBoneTransforms.Sort(FCompareBoneTransformIndex());
 	}
 }
 
@@ -138,13 +200,13 @@ bool FAnimNode_OrientationWarping::IsValidToEvaluate(const USkeleton* Skeleton,
 {
 	bool bIKFootRootIsValid = IKFootRootBoneIndex != INDEX_NONE;
 	bool bIKFeetAreValid = IKFootBoneIndexArray.Num() > 0;
-	for (auto& IKFootBoneIndex : IKFootBoneIndexArray)
+	for (const auto& IKFootBoneIndex : IKFootBoneIndexArray)
 	{
 		bIKFeetAreValid = bIKFeetAreValid && IKFootBoneIndex != INDEX_NONE;
 	}
 
 	bool bSpineIsValid = SpineBoneDataArray.Num() > 0;
-	for (auto& Spine : SpineBoneDataArray)
+	for (const auto& Spine : SpineBoneDataArray)
 	{
 		bSpineIsValid = bSpineIsValid && Spine.BoneIndex != INDEX_NONE;
 	}
@@ -223,4 +285,4 @@ void FAnimNode_OrientationWarping::InitializeBoneReferences(const FBoneContainer
 			}
 		}
 	}
-}
+}
\ No newline at end of file
diff --git a/Engine/Plugins/Experimental/Animation/AnimationWarping/Source/Runtime/Private/BoneControllers/AnimNode_SlopeWarping.cpp b/Engine/Plugins/Experimental/Animation/AnimationWarping/Source/Runtime/Private/BoneControllers/AnimNode_SlopeWarping.cpp
--- a/Engine/Plugins/Experimental/Animation/AnimationWarping/Source/Runtime/Private/BoneControllers/AnimNode_SlopeWarping.cpp
+++ b/Engine/Plugins/Experimental/Animation/AnimationWarping/Source/Runtime/Private/BoneControllers/AnimNode_SlopeWarping.cpp
@@ -21,6 +21,7 @@ FAnimNode_SlopeWarping::FAnimNode_SlopeWarping()
 	: MySkelMeshComponent(nullptr)
 	, MyMovementComponent(nullptr)
 	, MyAnimInstanceProxy(nullptr)
+	, Mode(EWarpingEvaluationMode::Manual)
 	, GravityDir(-FVector::UpVector)
 	, CustomFloorOffset(FVector::ZeroVector)
 	, CachedDeltaTime(0.f)
diff --git a/Engine/Plugins/Experimental/Animation/AnimationWarping/Source/Runtime/Private/BoneControllers/AnimNode_StrideWarping.cpp b/Engine/Plugins/Experimental/Animation/AnimationWarping/Source/Runtime/Private/BoneControllers/AnimNode_StrideWarping.cpp
--- a/Engine/Plugins/Experimental/Animation/AnimationWarping/Source/Runtime/Private/BoneControllers/AnimNode_StrideWarping.cpp
+++ b/Engine/Plugins/Experimental/Animation/AnimationWarping/Source/Runtime/Private/BoneControllers/AnimNode_StrideWarping.cpp
@@ -5,6 +5,7 @@
 #include "Animation/AnimInstanceProxy.h"
 #include "DrawDebugHelpers.h"
 #include "Animation/InputScaleBias.h"
+#include "Animation/AnimRootMotionProvider.h"
 
 #if ENABLE_ANIM_DEBUG
 TAutoConsoleVariable<int32> CVarAnimNodeStrideWarpingDebug(TEXT("a.AnimNode.StrideWarping.Debug"), 0, TEXT("Turn on debug for AnimNode_StrideWarping"));
@@ -15,11 +16,13 @@ TAutoConsoleVariable<int32> CVarStrideWarpingEnable(TEXT("a.AnimNode.StrideWarpi
 DECLARE_CYCLE_STAT(TEXT("StrideWarping Eval"), STAT_StrideWarping_Eval, STATGROUP_Anim);
 
 FAnimNode_StrideWarping::FAnimNode_StrideWarping()
-	: StrideWarpingAxisMode(EStrideWarpingAxisMode::ActorSpaceVectorInput)
+	: Mode(EWarpingEvaluationMode::Manual)
+	, ManualStrideWarpingDir(ForceInitToZero)
+	, ManualStrideScaling(1.f)
+	, LocomotionSpeed(0.f)
+	, StrideWarpingAxisMode(EStrideWarpingAxisMode::ActorSpaceVectorInput)
 	, FloorNormalAxisMode(EStrideWarpingAxisMode::IKFootRootLocalZ)
 	, GravityDirAxisMode(EStrideWarpingAxisMode::ComponentSpaceVectorInput)
-	, StrideScaling(1.f)
-	, ManualStrideWarpingDir(ForceInitToZero)
 	, ManualFloorNormalInput(ForceInitToZero)
 	, ManualGravityDirInput(-FVector::UpVector)
 	, PelvisPostAdjustmentAlpha(0.4f)
@@ -56,7 +59,7 @@ void FAnimNode_StrideWarping::UpdateInternal(const FAnimationUpdateContext& Cont
 {
 	FAnimNode_SkeletalControlBase::UpdateInternal(Context);
 
-	CachedDeltaTime += Context.GetDeltaTime();
+	CachedDeltaTime = Context.GetDeltaTime();
 }
 
 FVector FAnimNode_StrideWarping::GetAxisModeValue(const EStrideWarpingAxisMode& AxisMode, const FTransform& IKFootRootCSTransform, const FVector& UserSuppliedVector) const
@@ -83,10 +86,24 @@ FVector FAnimNode_StrideWarping::GetAxisModeValue(const EStrideWarpingAxisMode&
 void FAnimNode_StrideWarping::EvaluateSkeletalControl_AnyThread(FComponentSpacePoseContext& Output, TArray<FBoneTransform>& OutBoneTransforms)
 {
 	SCOPE_CYCLE_COUNTER(STAT_StrideWarping_Eval);
-
 	check(OutBoneTransforms.Num() == 0);
+
 	const FBoneContainer& RequiredBones = Output.Pose.GetPose().GetBoneContainer();
 
+	FTransform RootMotionTransformDelta;
+	const UE::Anim::IAnimRootMotionProvider* RootMotionProvider = UE::Anim::IAnimRootMotionProvider::Get();
+	bool bGraphDrivenWarping = RootMotionProvider && Mode == EWarpingEvaluationMode::Graph;
+
+	if (bGraphDrivenWarping)
+	{
+		// Graph driven stride warping will override the manual stride direction with the intent of the current animation sub-graph's accumulated root motion
+		bGraphDrivenWarping = RootMotionProvider->ExtractRootMotion(Output.CustomAttributes, RootMotionTransformDelta);
+		if (bGraphDrivenWarping)
+		{
+			ManualStrideWarpingDir = RootMotionTransformDelta.GetTranslation().GetSafeNormal2D();
+		}
+	}
+
 	FTransform IKFootRootTransform = Output.Pose.GetComponentSpaceTransform(IKFootRootBone.GetCompactPoseIndex(RequiredBones));
 	FVector StrideWarpingPlaneNormal = GetAxisModeValue(StrideWarpingAxisMode, IKFootRootTransform, ManualStrideWarpingDir);
 	const FVector FloorPlaneNormal = GetAxisModeValue(FloorNormalAxisMode, IKFootRootTransform, ManualFloorNormalInput);
@@ -124,8 +141,38 @@ void FAnimNode_StrideWarping::EvaluateSkeletalControl_AnyThread(FComponentSpaceP
 #endif
 	}
 
+	float ComputedStrideScaling = ManualStrideScaling;
+	if (bGraphDrivenWarping)
+	{
+		// Early out when stride warping will have minimal pose contribution
+		if (FMath::IsNearlyZero(CachedDeltaTime, KINDA_SMALL_NUMBER))
+		{
+			return;
+		}
+
+		const float RootMotionSpeed = RootMotionTransformDelta.GetTranslation().Size() / CachedDeltaTime;
+		if (FMath::IsNearlyZero(RootMotionSpeed, KINDA_SMALL_NUMBER))
+		{
+			return;
+		}
+
+		// Graph driven stride scale factor will be determined by the ratio of the
+		// locomotion (capsule/physics) speed against the animation root motion speed
+		ComputedStrideScaling = LocomotionSpeed / RootMotionSpeed;
+	}
+
+	// Allow the opportunity for stride scale clamping and biasing regardless of evaluation mode
+	ComputedStrideScaling = StrideScalingScaleBiasClamp.ApplyTo(ComputedStrideScaling, CachedDeltaTime);
+
+	if (bGraphDrivenWarping)
+	{
+		// Forward the side effects of stride warping on the root motion contribution for this sub-graph
+		RootMotionTransformDelta.ScaleTranslation(ComputedStrideScaling);
+		const bool bRootMotionOverridden = RootMotionProvider->OverrideRootMotion(RootMotionTransformDelta, Output.CustomAttributes);
+		ensure(bRootMotionOverridden);
+	}
+
 	// Scale IK feet bones along Stride Warping Axis, from the hip bone location.
-	const float ActualStrideScaling = StrideScalingScaleBiasClamp.ApplyTo(StrideScaling, CachedDeltaTime);
 	for (auto& FootData : FeetData)
 	{
 		// Stride Warping along Stride Warping Axis
@@ -140,7 +187,7 @@ void FAnimNode_StrideWarping::EvaluateSkeletalControl_AnyThread(FComponentSpaceP
 		const FVector ScaleOrigin = FVector::PointPlaneProject(IKFootLocation, StrideWarpingPlaneOrigin, StrideWarpingPlaneNormal);
 
 		// Now the ScaleOrigin and IKFootLocation are forming a line parallel to the floor, and we can scale the IK foot.
-		const FVector WarpedLocation = ScaleOrigin + (IKFootLocation - ScaleOrigin) * ActualStrideScaling;
+		const FVector WarpedLocation = ScaleOrigin + (IKFootLocation - ScaleOrigin) * ComputedStrideScaling;
 		FootData.IKBoneTransform.SetLocation(WarpedLocation);
 
 #if ENABLE_ANIM_DEBUG
@@ -296,9 +343,6 @@ void FAnimNode_StrideWarping::EvaluateSkeletalControl_AnyThread(FComponentSpaceP
 
 	// Sort OutBoneTransforms so indices are in increasing order.
 	OutBoneTransforms.Sort(FCompareBoneTransformIndex());
-
-	// Clear time accumulator, to be filled during next update.
-	CachedDeltaTime = 0.f;
 }
 
 bool FAnimNode_StrideWarping::IsValidToEvaluate(const USkeleton* Skeleton, const FBoneContainer& RequiredBones)
@@ -307,7 +351,7 @@ bool FAnimNode_StrideWarping::IsValidToEvaluate(const USkeleton* Skeleton, const
 	return bIsEnabled && (FeetData.Num() > 0) 
 		&& (PelvisBone.GetCompactPoseIndex(RequiredBones) != INDEX_NONE) 
 		&& (IKFootRootBone.GetCompactPoseIndex(RequiredBones) != INDEX_NONE) 
-		&& (!FMath::IsNearlyEqual(StrideScalingScaleBiasClamp.ApplyTo(StrideScaling, 0.f), 1.f, 0.001f) || PelvisAdjustmentInterp.IsInMotion());
+		&& (!FMath::IsNearlyEqual(StrideScalingScaleBiasClamp.ApplyTo(ManualStrideScaling, 0.f), 1.f, 0.001f) || PelvisAdjustmentInterp.IsInMotion() || Mode == EWarpingEvaluationMode::Graph);
 }
 
 FCompactPoseBoneIndex FAnimNode_StrideWarping::FindHipBoneIndex(const FCompactPoseBoneIndex& InFootBoneIndex, const int32& NumBonesInLimb, const FBoneContainer& RequiredBones) const
diff --git a/Engine/Plugins/Experimental/Animation/AnimationWarping/Source/Runtime/Public/BoneControllers/AnimNode_OrientationWarping.h b/Engine/Plugins/Experimental/Animation/AnimationWarping/Source/Runtime/Public/BoneControllers/AnimNode_OrientationWarping.h
--- a/Engine/Plugins/Experimental/Animation/AnimationWarping/Source/Runtime/Public/BoneControllers/AnimNode_OrientationWarping.h
+++ b/Engine/Plugins/Experimental/Animation/AnimationWarping/Source/Runtime/Public/BoneControllers/AnimNode_OrientationWarping.h
@@ -2,6 +2,7 @@
 
 #pragma once
 
+#include "BoneControllers/BoneControllerTypes.h"
 #include "BoneControllers/AnimNode_SkeletalControlBase.h"
 #include "Animation/AnimNodeBase.h"
 #include "AnimNode_OrientationWarping.generated.h"
@@ -49,12 +50,12 @@ struct FOrientationWarpingSettings
 	GENERATED_USTRUCT_BODY()
 
 	/** Rotation Axis used to rotate mesh. */
-	UPROPERTY(EditAnywhere, Category = "Settings")
-	TEnumAsByte<EAxis::Type> YawRotationAxis;
+	UPROPERTY(EditAnywhere, Category = Settings)
+	TEnumAsByte<EAxis::Type> YawRotationAxis = EAxis::Z;
 
 	/** How much of rotation is on whole body versus on IK Feet. */
 	UPROPERTY(EditAnywhere, Category = "Settings")
-	float BodyOrientationAlpha;
+	float BodyOrientationAlpha = 0.5f;
 
 	/** Spine bones countering the rotation of the body, so character keeps aiming straight ahead. */
 	UPROPERTY(EditAnywhere, Category = "Settings")
@@ -67,11 +68,6 @@ struct FOrientationWarpingSettings
 	/** IK Foot Bones. To be rotated if 'BodyOrientationAlpha' is less than 1.f */
 	UPROPERTY(EditAnywhere, Category = "Settings")
 	TArray<FBoneReference> IKFootBones;
-
-	FOrientationWarpingSettings()
-		: YawRotationAxis(EAxis::Z)
-		, BodyOrientationAlpha(0.5f)
-	{}
 };
 
 USTRUCT(BlueprintInternalUseOnly)
@@ -79,10 +75,13 @@ struct ANIMATIONWARPINGRUNTIME_API FAnimNode_OrientationWarping : public FAnimNo
 {
 	GENERATED_USTRUCT_BODY()
 
-	UPROPERTY(EditAnywhere, BlueprintReadWrite, Category = "Settings", meta = (PinShownByDefault))
+	UPROPERTY(EditAnywhere, BlueprintReadWrite, Category = Evaluation, meta = (PinHiddenByDefault))
+	EWarpingEvaluationMode Mode;
+
+	UPROPERTY(EditAnywhere, BlueprintReadWrite, Category = Settings, meta = (PinShownByDefault))
 	float LocomotionAngle;
 
-	UPROPERTY(EditAnywhere, BlueprintReadWrite, Category = "Settings", meta = (PinHiddenByDefault))
+	UPROPERTY(EditAnywhere, BlueprintReadWrite, Category = Settings, meta = (PinHiddenByDefault))
 	FOrientationWarpingSettings Settings;
 
 	UPROPERTY(Transient)
@@ -92,9 +91,6 @@ struct ANIMATIONWARPINGRUNTIME_API FAnimNode_OrientationWarping : public FAnimNo
 
 	FCompactPoseBoneIndex IKFootRootBoneIndex;
 
-	UPROPERTY(Transient)
-	float CachedDeltaTime;
-
 public:
 	FAnimNode_OrientationWarping();
 
diff --git a/Engine/Plugins/Experimental/Animation/AnimationWarping/Source/Runtime/Public/BoneControllers/AnimNode_SlopeWarping.h b/Engine/Plugins/Experimental/Animation/AnimationWarping/Source/Runtime/Public/BoneControllers/AnimNode_SlopeWarping.h
--- a/Engine/Plugins/Experimental/Animation/AnimationWarping/Source/Runtime/Public/BoneControllers/AnimNode_SlopeWarping.h
+++ b/Engine/Plugins/Experimental/Animation/AnimationWarping/Source/Runtime/Public/BoneControllers/AnimNode_SlopeWarping.h
@@ -2,6 +2,7 @@
 
 #pragma once
 
+#include "BoneControllers/BoneControllerTypes.h"
 #include "BoneControllers/AnimNode_SkeletalControlBase.h"
 #include "Engine/SpringInterpolator.h"
 #include "AnimNode_SlopeWarping.generated.h"
@@ -71,6 +72,9 @@ struct ANIMATIONWARPINGRUNTIME_API FAnimNode_SlopeWarping : public FAnimNode_Ske
 	UCharacterMovementComponent* MyMovementComponent;
 	FAnimInstanceProxy* MyAnimInstanceProxy;
 
+	UPROPERTY(VisibleDefaultsOnly, BlueprintReadOnly, Category = Evaluation, meta = (PinHiddenByDefault))
+	EWarpingEvaluationMode Mode;
+
 	/** IKFoot Root Bone. **/
 	UPROPERTY(EditAnywhere, Category = "Settings")
 	FBoneReference IKFootRootBone;
diff --git a/Engine/Plugins/Experimental/Animation/AnimationWarping/Source/Runtime/Public/BoneControllers/AnimNode_StrideWarping.h b/Engine/Plugins/Experimental/Animation/AnimationWarping/Source/Runtime/Public/BoneControllers/AnimNode_StrideWarping.h
--- a/Engine/Plugins/Experimental/Animation/AnimationWarping/Source/Runtime/Public/BoneControllers/AnimNode_StrideWarping.h
+++ b/Engine/Plugins/Experimental/Animation/AnimationWarping/Source/Runtime/Public/BoneControllers/AnimNode_StrideWarping.h
@@ -2,6 +2,7 @@
 
 #pragma once
 
+#include "BoneControllers/BoneControllerTypes.h"
 #include "BoneControllers/AnimNode_SkeletalControlBase.h"
 #include "Engine/SpringInterpolator.h"
 #include "AnimNode_StrideWarping.generated.h"
@@ -66,63 +67,67 @@ struct ANIMATIONWARPINGRUNTIME_API FAnimNode_StrideWarping : public FAnimNode_Sk
 
 	FAnimInstanceProxy* MyAnimInstanceProxy;
 
-	/** IKFoot Root Bone. **/
-	UPROPERTY(EditAnywhere, Category = "Settings")
+	UPROPERTY(EditAnywhere, BlueprintReadWrite, Category = Evaluation, meta = (PinHiddenByDefault))
+	EWarpingEvaluationMode Mode;
+
+	UPROPERTY(EditAnywhere, BlueprintReadWrite, Category = Evaluation, meta = (PinHiddenByDefault))
+	FVector ManualStrideWarpingDir;
+
+	UPROPERTY(EditAnywhere, BlueprintReadWrite, Category = Evaluation, meta = (PinHiddenByDefault))
+	float ManualStrideScaling;
+
+	UPROPERTY(EditAnywhere, BlueprintReadWrite, Category = Evaluation, meta = (PinHiddenByDefault))
+	float LocomotionSpeed;
+
+	UPROPERTY(EditAnywhere, Category = Settings)
 	FBoneReference IKFootRootBone;
 
-	UPROPERTY(EditAnywhere, Category = "Settings")
+	UPROPERTY(EditAnywhere, Category = Settings)
 	TArray<FStrideWarpingFootDefinition> FeetDefinitions;
 
 	UPROPERTY(Transient)
 	TArray<FStrideWarpingFootData> FeetData;
 
-	/** Pelvis Bone. **/
-	UPROPERTY(EditAnywhere, Category = "Settings")
+	UPROPERTY(EditAnywhere, Category = Settings)
 	FBoneReference PelvisBone;
 
-	UPROPERTY(EditAnywhere, BlueprintReadWrite, Category = "Settings")
+	UPROPERTY(EditAnywhere, BlueprintReadWrite, Category = Settings)
 	EStrideWarpingAxisMode StrideWarpingAxisMode;
 
-	UPROPERTY(EditAnywhere, BlueprintReadWrite, Category = "Settings")
+	UPROPERTY(EditAnywhere, BlueprintReadWrite, Category = Settings)
 	EStrideWarpingAxisMode FloorNormalAxisMode;
 
 	/** Direction vector for adjusting hips up and down. */
-	UPROPERTY(EditAnywhere, BlueprintReadWrite, Category = "Settings")
+	UPROPERTY(EditAnywhere, BlueprintReadWrite, Category = Settings)
 	EStrideWarpingAxisMode GravityDirAxisMode;
 
-	UPROPERTY(EditAnywhere, BlueprintReadWrite, Category = "Settings", meta = (PinShownByDefault))
-	float StrideScaling;
-
 	// Additional scaling, offsetting and clamping of PlayRate input.
 	// Performed after PlayRateBasis.
 	UPROPERTY(EditAnywhere, BlueprintReadWrite, Category = Settings)
 	FInputScaleBiasClamp StrideScalingScaleBiasClamp;
 
-	UPROPERTY(EditAnywhere, BlueprintReadWrite, Category = "Settings", meta = (PinHiddenByDefault))
-	FVector ManualStrideWarpingDir;
-
-	UPROPERTY(EditAnywhere, BlueprintReadWrite, Category = "Settings", meta = (PinHiddenByDefault))
+	UPROPERTY(EditAnywhere, BlueprintReadWrite, Category = Settings, meta = (PinHiddenByDefault))
 	FVector ManualFloorNormalInput;
 
-	UPROPERTY(EditAnywhere, BlueprintReadWrite, Category = "Settings", meta = (PinHiddenByDefault))
+	UPROPERTY(EditAnywhere, BlueprintReadWrite, Category = Settings, meta = (PinHiddenByDefault))
 	FVector ManualGravityDirInput;
 
-	UPROPERTY(EditAnywhere, BlueprintReadWrite, Category = "Settings")
+	UPROPERTY(EditAnywhere, BlueprintReadWrite, Category = Settings)
 	float PelvisPostAdjustmentAlpha;
 
-	UPROPERTY(EditAnywhere, BlueprintReadWrite, Category = "Settings")
+	UPROPERTY(EditAnywhere, BlueprintReadWrite, Category = Settings)
 	int32 PelvisAdjustmentMaxIter;
 
-	UPROPERTY(EditAnywhere, Category = "Settings")
+	UPROPERTY(EditAnywhere, Category = Settings)
 	FVectorRK4SpringInterpolator PelvisAdjustmentInterp;
 
-	UPROPERTY(EditAnywhere, Category = "Settings")
+	UPROPERTY(EditAnywhere, Category = Settings)
 	uint32 bAdjustThighBonesRotation : 1;
 
-	UPROPERTY(EditAnywhere, Category = "Settings")
+	UPROPERTY(EditAnywhere, Category = Settings)
 	uint32 bClampIKUsingFKLeg : 1;
 
-	UPROPERTY(EditAnywhere, Category = "Settings")
+	UPROPERTY(EditAnywhere, Category = Settings)
 	uint32 bOrientStrideWarpingAxisBasedOnFloorNormal : 1;
 
 	UPROPERTY(Transient)
diff --git a/Engine/Plugins/Experimental/Animation/PoseSearch/Source/Editor/Private/AnimGraphNode_MotionMatching.cpp b/Engine/Plugins/Experimental/Animation/PoseSearch/Source/Editor/Private/AnimGraphNode_MotionMatching.cpp
--- a/Engine/Plugins/Experimental/Animation/PoseSearch/Source/Editor/Private/AnimGraphNode_MotionMatching.cpp
+++ b/Engine/Plugins/Experimental/Animation/PoseSearch/Source/Editor/Private/AnimGraphNode_MotionMatching.cpp
@@ -1,6 +1,7 @@
 // Copyright Epic Games, Inc. All Rights Reserved.
 
 #include "AnimGraphNode_MotionMatching.h"
+#include "Animation/AnimRootMotionProvider.h"
 
 #define LOCTEXT_NAMESPACE "AnimGraphNode_MotionMatching"
 
@@ -25,5 +26,12 @@ FText UAnimGraphNode_MotionMatching::GetMenuCategory() const
 	return LOCTEXT("NodeCategory", "Pose Search");
 }
 
+void UAnimGraphNode_MotionMatching::GetOutputLinkAttributes(FNodeAttributeArray& OutAttributes) const
+{
+	if (UE::Anim::IAnimRootMotionProvider::Get())
+	{
+		OutAttributes.Add(UE::Anim::IAnimRootMotionProvider::RootMotionDeltaAttributeName);
+	}
+}
 
 #undef LOCTEXT_NAMESPACE
diff --git a/Engine/Plugins/Experimental/Animation/PoseSearch/Source/Editor/Public/AnimGraphNode_MotionMatching.h b/Engine/Plugins/Experimental/Animation/PoseSearch/Source/Editor/Public/AnimGraphNode_MotionMatching.h
--- a/Engine/Plugins/Experimental/Animation/PoseSearch/Source/Editor/Public/AnimGraphNode_MotionMatching.h
+++ b/Engine/Plugins/Experimental/Animation/PoseSearch/Source/Editor/Public/AnimGraphNode_MotionMatching.h
@@ -21,4 +21,5 @@ class UAnimGraphNode_MotionMatching : public UAnimGraphNode_Base
 	virtual FText GetTooltipText() const override;
 	virtual FText GetNodeTitle(ENodeTitleType::Type TitleType) const override;
 	virtual FText GetMenuCategory() const override;
+	virtual void GetOutputLinkAttributes(FNodeAttributeArray& OutAttributes) const override;
 };
diff --git a/Engine/Plugins/Experimental/Animation/PoseSearch/Source/Runtime/Public/PoseSearch/AnimNode_MotionMatching.h b/Engine/Plugins/Experimental/Animation/PoseSearch/Source/Runtime/Public/PoseSearch/AnimNode_MotionMatching.h
--- a/Engine/Plugins/Experimental/Animation/PoseSearch/Source/Runtime/Public/PoseSearch/AnimNode_MotionMatching.h
+++ b/Engine/Plugins/Experimental/Animation/PoseSearch/Source/Runtime/Public/PoseSearch/AnimNode_MotionMatching.h
@@ -17,7 +17,7 @@ public:
 	FPoseLink Source;
 
 	// Collection of animations for motion matching
-	UPROPERTY(EditAnywhere, BlueprintReadWrite, Category=Settings, meta=(PinShownByDefault))
+	UPROPERTY(EditAnywhere, BlueprintReadWrite, Category=Settings, meta=(PinHiddenByDefault))
 	const UPoseSearchDatabase* Database = nullptr;
 
 	// Motion matching goal
diff --git a/Engine/Source/Editor/AnimGraph/Private/AnimGraphModule.cpp b/Engine/Source/Editor/AnimGraph/Private/AnimGraphModule.cpp
--- a/Engine/Source/Editor/AnimGraph/Private/AnimGraphModule.cpp
+++ b/Engine/Source/Editor/AnimGraph/Private/AnimGraphModule.cpp
@@ -28,6 +28,7 @@
 #include "AnimGraphAttributes.h"
 #include "Animation/AnimSync.h"
 #include "Animation/AnimNode_Inertialization.h"
+#include "Animation/AnimRootMotionProvider.h"
 
 IMPLEMENT_MODULE(FAnimGraphModule, AnimGraph);
 
@@ -107,6 +108,17 @@ void FAnimGraphModule::StartupModule()
 			EAnimGraphAttributesDisplayMode::Automatic)
 	);
 
+	AnimGraphAttributes->Register(
+		FAnimGraphAttributeDesc(
+			UE::Anim::IAnimRootMotionProvider::RootMotionDeltaAttributeName,
+			EAnimGraphAttributeBlend::Blendable,
+			*FEditorStyle::Get().GetBrush("AnimGraph.Attribute.RootMotionDelta.Icon"),
+			LOCTEXT("RootMotionDelta", "Root Motion Delta"),
+			LOCTEXT("RootMotionDeltaToolTip", "Root Motion Delta: Per-frame root motion delta transform, blendable"),
+			FEditorStyle::Get().GetSlateColor("AnimGraph.Attribute.RootMotionDelta.Color"),
+			EAnimGraphAttributesDisplayMode::Automatic)
+	);
+
 	AnimGraphAttributes->Register(
 		FAnimGraphAttributeDesc(
 			UE::Anim::IInertializationRequester::Attribute,
diff --git a/Engine/Source/Editor/AnimGraph/Private/AnimGraphNode_BlendSpaceBase.cpp b/Engine/Source/Editor/AnimGraph/Private/AnimGraphNode_BlendSpaceBase.cpp
--- a/Engine/Source/Editor/AnimGraph/Private/AnimGraphNode_BlendSpaceBase.cpp
+++ b/Engine/Source/Editor/AnimGraph/Private/AnimGraphNode_BlendSpaceBase.cpp
@@ -7,6 +7,7 @@
 #include "Animation/AimOffsetBlendSpace.h"
 #include "Animation/BlendSpace1D.h"
 #include "Animation/AimOffsetBlendSpace1D.h"
+#include "Animation/AnimRootMotionProvider.h"
 
 #define LOCTEXT_NAMESPACE "AnimGraphNode_BlendSpaceBase"
 
@@ -79,6 +80,16 @@ void UAnimGraphNode_BlendSpaceBase::PostProcessPinName(const UEdGraphPin* Pin, F
 	Super::PostProcessPinName(Pin, DisplayName);
 }
 
+void UAnimGraphNode_BlendSpaceBase::GetOutputLinkAttributes(FNodeAttributeArray& OutAttributes) const
+{
+	Super::GetOutputLinkAttributes(OutAttributes);
+
+	if (UE::Anim::IAnimRootMotionProvider::Get())
+	{
+		OutAttributes.Add(UE::Anim::IAnimRootMotionProvider::RootMotionDeltaAttributeName);
+	}
+}
+
 FText UAnimGraphNode_BlendSpaceBase::GetMenuCategory() const
 {
 	return LOCTEXT("BlendSpaceCategory_Label", "BlendSpaces");
diff --git a/Engine/Source/Editor/AnimGraph/Private/AnimGraphNode_RandomPlayer.cpp b/Engine/Source/Editor/AnimGraph/Private/AnimGraphNode_RandomPlayer.cpp
--- a/Engine/Source/Editor/AnimGraph/Private/AnimGraphNode_RandomPlayer.cpp
+++ b/Engine/Source/Editor/AnimGraph/Private/AnimGraphNode_RandomPlayer.cpp
@@ -4,6 +4,7 @@
 
 #include "EditorCategoryUtils.h"
 #include "Animation/AnimAttributes.h"
+#include "Animation/AnimRootMotionProvider.h"
 
 #define LOCTEXT_NAMESPACE "AnimGraphNode_RandomPlayer"
 
@@ -31,6 +32,11 @@ void UAnimGraphNode_RandomPlayer::GetOutputLinkAttributes(FNodeAttributeArray& O
 {
 	OutAttributes.Add(UE::Anim::FAttributes::Curves);
 	OutAttributes.Add(UE::Anim::FAttributes::Attributes);
+
+	if (UE::Anim::IAnimRootMotionProvider::Get())
+	{
+		OutAttributes.Add(UE::Anim::IAnimRootMotionProvider::RootMotionDeltaAttributeName);
+	}
 }
 
 #undef LOCTEXT_NAMESPACE
diff --git a/Engine/Source/Editor/AnimGraph/Private/AnimGraphNode_SequencePlayer.cpp b/Engine/Source/Editor/AnimGraph/Private/AnimGraphNode_SequencePlayer.cpp
--- a/Engine/Source/Editor/AnimGraph/Private/AnimGraphNode_SequencePlayer.cpp
+++ b/Engine/Source/Editor/AnimGraph/Private/AnimGraphNode_SequencePlayer.cpp
@@ -17,6 +17,7 @@
 #include "Animation/AnimComposite.h"
 #include "Animation/AnimSequence.h"
 #include "Animation/AnimPoseSearchProvider.h"
+#include "Animation/AnimRootMotionProvider.h"
 #include "UObject/UE5MainStreamObjectVersion.h"
 #include "IAnimBlueprintNodeOverrideAssetsContext.h"
 
@@ -369,6 +370,16 @@ EAnimAssetHandlerType UAnimGraphNode_SequencePlayer::SupportsAssetClass(const UC
 	}
 }
 
+void UAnimGraphNode_SequencePlayer::GetOutputLinkAttributes(FNodeAttributeArray& OutAttributes) const
+{
+	Super::GetOutputLinkAttributes(OutAttributes);
+
+	if (UE::Anim::IAnimRootMotionProvider::Get())
+	{
+		OutAttributes.Add(UE::Anim::IAnimRootMotionProvider::RootMotionDeltaAttributeName);
+	}
+}
+
 void UAnimGraphNode_SequencePlayer::ValidateAnimNodeDuringCompilation(class USkeleton* ForSkeleton, class FCompilerResultsLog& MessageLog)
 {
 	Super::ValidateAnimNodeDuringCompilation(ForSkeleton, MessageLog);
diff --git a/Engine/Source/Editor/AnimGraph/Public/AnimGraphNode_BlendSpaceBase.h b/Engine/Source/Editor/AnimGraph/Public/AnimGraphNode_BlendSpaceBase.h
--- a/Engine/Source/Editor/AnimGraph/Public/AnimGraphNode_BlendSpaceBase.h
+++ b/Engine/Source/Editor/AnimGraph/Public/AnimGraphNode_BlendSpaceBase.h
@@ -22,6 +22,7 @@ class UAnimGraphNode_BlendSpaceBase : public UAnimGraphNode_AssetPlayerBase
 	ANIMGRAPH_API virtual void CustomizePinData(UEdGraphPin* Pin, FName SourcePropertyName, int32 ArrayIndex) const override;
 	ANIMGRAPH_API virtual void PreloadRequiredAssets() override;
 	ANIMGRAPH_API virtual void PostProcessPinName(const UEdGraphPin* Pin, FString& DisplayName) const override;
+	ANIMGRAPH_API virtual void GetOutputLinkAttributes(FNodeAttributeArray& OutAttributes) const;
 	// End of UAnimGraphNode_Base interface
 
 protected:
diff --git a/Engine/Source/Editor/AnimGraph/Public/AnimGraphNode_SequencePlayer.h b/Engine/Source/Editor/AnimGraph/Public/AnimGraphNode_SequencePlayer.h
--- a/Engine/Source/Editor/AnimGraph/Public/AnimGraphNode_SequencePlayer.h
+++ b/Engine/Source/Editor/AnimGraph/Public/AnimGraphNode_SequencePlayer.h
@@ -46,6 +46,7 @@ class UAnimGraphNode_SequencePlayer : public UAnimGraphNode_AssetPlayerBase
 	virtual bool IsActionFilteredOut(class FBlueprintActionFilter const& Filter) override;
 	virtual EAnimAssetHandlerType SupportsAssetClass(const UClass* AssetClass) const override;
 	virtual void OnOverrideAssets(IAnimBlueprintNodeOverrideAssetsContext& InContext) const override;
+	virtual void GetOutputLinkAttributes(FNodeAttributeArray& OutAttributes) const override;
 	// End of UAnimGraphNode_Base interface
 
 	// UK2Node interface
diff --git a/Engine/Source/Editor/EditorStyle/Private/StarshipStyle.cpp b/Engine/Source/Editor/EditorStyle/Private/StarshipStyle.cpp
--- a/Engine/Source/Editor/EditorStyle/Private/StarshipStyle.cpp
+++ b/Engine/Source/Editor/EditorStyle/Private/StarshipStyle.cpp
@@ -5716,6 +5716,8 @@ void FStarshipEditorStyle::FStyle::SetupPersonaStyle()
 		Set("AnimGraph.Attribute.Attributes.Color", FStyleColors::AccentYellow);
 		Set("AnimGraph.Attribute.Sync.Icon", new IMAGE_BRUSH("Icons/icon_tab_TimecodeProvider_16x", Icon10x10));
 		Set("AnimGraph.Attribute.Sync.Color", FStyleColors::AccentBlue);
+		Set("AnimGraph.Attribute.RootMotionDelta.Icon", new IMAGE_BRUSH("Icons/icon_axis_local_16px", Icon12x12));
+		Set("AnimGraph.Attribute.RootMotionDelta.Color", FStyleColors::AccentGreen);
 		Set("AnimGraph.Attribute.InertialBlending.Icon", new IMAGE_BRUSH("Icons/icon_tab_TimecodeProvider_16x", Icon10x10));
 		Set("AnimGraph.Attribute.InertialBlending.Color", FStyleColors::AccentOrange);
 	}	
diff --git a/Engine/Source/Runtime/AnimGraphRuntime/Private/AnimNodes/AnimNode_BlendSpaceGraphBase.cpp b/Engine/Source/Runtime/AnimGraphRuntime/Private/AnimNodes/AnimNode_BlendSpaceGraphBase.cpp
--- a/Engine/Source/Runtime/AnimGraphRuntime/Private/AnimNodes/AnimNode_BlendSpaceGraphBase.cpp
+++ b/Engine/Source/Runtime/AnimGraphRuntime/Private/AnimNodes/AnimNode_BlendSpaceGraphBase.cpp
@@ -87,7 +87,7 @@ void FAnimNode_BlendSpaceGraphBase::Evaluate_AnyThread(FPoseContext& Output)
 
 	check(BlendSpace != nullptr);
 
-	BlendSpace->GetAnimationPose(BlendSampleDataCache, SamplePoseLinks, Output);
+	BlendSpace->GetAnimationPose(BlendSampleDataCache, SamplePoseLinks, {}/** unused parameter **/, Output);
 }
 
 void FAnimNode_BlendSpaceGraphBase::GatherDebugData(FNodeDebugData& DebugData)
diff --git a/Engine/Source/Runtime/AnimGraphRuntime/Private/AnimNodes/AnimNode_BlendSpacePlayer.cpp b/Engine/Source/Runtime/AnimGraphRuntime/Private/AnimNodes/AnimNode_BlendSpacePlayer.cpp
--- a/Engine/Source/Runtime/AnimGraphRuntime/Private/AnimNodes/AnimNode_BlendSpacePlayer.cpp
+++ b/Engine/Source/Runtime/AnimGraphRuntime/Private/AnimNodes/AnimNode_BlendSpacePlayer.cpp
@@ -87,6 +87,8 @@ void FAnimNode_BlendSpacePlayer::UpdateInternal(const FAnimationUpdateContext& C
 
 		FAnimTickRecord TickRecord(CurrentBlendSpace, Position, BlendSampleDataCache, BlendFilter, GetLoop(), GetPlayRate(), Context.GetFinalBlendWeight(), /*inout*/ InternalTimeAccumulator, MarkerTickRecord);
 		TickRecord.RootMotionWeightModifier = Context.GetRootMotionWeightModifier();
+		TickRecord.DeltaTimeRecord = &DeltaTimeRecord;
+
 		UE::Anim::FAnimSyncParams SyncParams(GetGroupName(), GetGroupRole(), GetGroupMethod());
 		if(Context.GetSharedContext())
 		{
@@ -121,7 +123,7 @@ void FAnimNode_BlendSpacePlayer::Evaluate_AnyThread(FPoseContext& Output)
 	if ((CurrentBlendSpace != nullptr) && (Output.AnimInstanceProxy->IsSkeletonCompatible(CurrentBlendSpace->GetSkeleton())))
 	{
 		FAnimationPoseData AnimationPoseData(Output);
-		CurrentBlendSpace->GetAnimationPose(BlendSampleDataCache, AnimationPoseData);
+		CurrentBlendSpace->GetAnimationPose(BlendSampleDataCache, GetLoop(), AnimationPoseData);
 	}
 	else
 	{
diff --git a/Engine/Source/Runtime/AnimGraphRuntime/Private/AnimNodes/AnimNode_RandomPlayer.cpp b/Engine/Source/Runtime/AnimGraphRuntime/Private/AnimNodes/AnimNode_RandomPlayer.cpp
--- a/Engine/Source/Runtime/AnimGraphRuntime/Private/AnimNodes/AnimNode_RandomPlayer.cpp
+++ b/Engine/Source/Runtime/AnimGraphRuntime/Private/AnimNodes/AnimNode_RandomPlayer.cpp
@@ -129,7 +129,7 @@ void FAnimNode_RandomPlayer::Update_AnyThread(const FAnimationUpdateContext& Con
 
 	// If we looped around, adjust the previous play time to always be before the current playtime,
 	// since we can assume modulo. This makes the crossing check for the start time a lot simpler.
-	float AdjustedPreviousPlayTime = CurrentData->PreviousPlayTime;
+	float AdjustedPreviousPlayTime = CurrentData->DeltaTimeRecord.Previous;
 	if (CurrentData->CurrentPlayTime < AdjustedPreviousPlayTime)
 	{
 		AdjustedPreviousPlayTime -= CurrentSequence->GetPlayLength();
@@ -218,8 +218,8 @@ void FAnimNode_RandomPlayer::Update_AnyThread(const FAnimationUpdateContext& Con
 	}
 
 	// Cache time to detect loops
-	CurrentData->PreviousPlayTime = CurrentData->CurrentPlayTime;
-	NextData->PreviousPlayTime = NextData->CurrentPlayTime;
+	CurrentData->DeltaTimeRecord.Previous = CurrentData->CurrentPlayTime;
+	NextData->DeltaTimeRecord.Previous = NextData->CurrentPlayTime;
 
 	if (bAdvanceToNextEntry)
 	{
@@ -231,7 +231,8 @@ void FAnimNode_RandomPlayer::Update_AnyThread(const FAnimationUpdateContext& Con
 	}
 
 	FAnimTickRecord TickRecord(CurrentData->Entry->Sequence, true, CurrentData->PlayRate, CurrentData->BlendWeight, CurrentData->CurrentPlayTime, CurrentData->MarkerTickRecord);
-	
+	TickRecord.DeltaTimeRecord = &CurrentData->DeltaTimeRecord;
+
 	if(Context.GetSharedContext())
 	{
 		Context.GetSharedContext()->MessageStack.MakeEventContextData(TickRecord.ContextData);
@@ -244,6 +245,7 @@ void FAnimNode_RandomPlayer::Update_AnyThread(const FAnimationUpdateContext& Con
 	if (FAnimationRuntime::HasWeight(NextData->BlendWeight))
 	{
 		FAnimTickRecord NextTickRecord(NextData->Entry->Sequence, true, NextData->PlayRate, NextData->BlendWeight, NextData->CurrentPlayTime, NextData->MarkerTickRecord);
+		NextTickRecord.DeltaTimeRecord = &NextData->DeltaTimeRecord;
 
 		if(Context.GetSharedContext())
 		{
@@ -296,12 +298,11 @@ void FAnimNode_RandomPlayer::Evaluate_AnyThread(FPoseContext& Output)
 
 		UAnimSequence* NextSequence = NextData.Entry->Sequence;
 
-
 		FAnimationPoseData CurrentPoseData(Poses[0], Curves[0], Attributes[0]);
 		FAnimationPoseData NextPoseData(Poses[1], Curves[1], Attributes[1]);
 
-		CurrentSequence->GetAnimationPose(CurrentPoseData, FAnimExtractContext(CurrentData.CurrentPlayTime, AnimProxy->ShouldExtractRootMotion()));
-		NextSequence->GetAnimationPose(NextPoseData, FAnimExtractContext(NextData.CurrentPlayTime, AnimProxy->ShouldExtractRootMotion()));
+		CurrentSequence->GetAnimationPose(CurrentPoseData, FAnimExtractContext(CurrentData.CurrentPlayTime, AnimProxy->ShouldExtractRootMotion(), CurrentData.DeltaTimeRecord, CurrentData.RemainingLoops > 0));
+		NextSequence->GetAnimationPose(NextPoseData, FAnimExtractContext(NextData.CurrentPlayTime, AnimProxy->ShouldExtractRootMotion(), NextData.DeltaTimeRecord, NextData.RemainingLoops > 0));
 
 		FAnimationPoseData AnimationPoseData(Output);
 		FAnimationRuntime::BlendPosesTogether(Poses, Curves, Attributes, Weights, AnimationPoseData);
@@ -310,7 +311,7 @@ void FAnimNode_RandomPlayer::Evaluate_AnyThread(FPoseContext& Output)
 	{
 		// Single animation, no blending needed.
 		FAnimationPoseData AnimationPoseData(Output);
-		CurrentSequence->GetAnimationPose(AnimationPoseData, FAnimExtractContext(CurrentData.CurrentPlayTime, Output.AnimInstanceProxy->ShouldExtractRootMotion()));
+		CurrentSequence->GetAnimationPose(AnimationPoseData, FAnimExtractContext(CurrentData.CurrentPlayTime, Output.AnimInstanceProxy->ShouldExtractRootMotion(), CurrentData.DeltaTimeRecord, CurrentData.RemainingLoops > 0));
 	}
 }
 
@@ -375,7 +376,8 @@ void FAnimNode_RandomPlayer::InitPlayData(FRandomAnimPlayData& Data, int32 Valid
 
 	Data.PlayStartTime = 0.0f;
 	Data.CurrentPlayTime = 0.0f;
-	Data.PreviousPlayTime = 0.0f;
+	Data.DeltaTimeRecord.Previous = 0.0f;
+	Data.DeltaTimeRecord.Delta = 0.0f;
 	Data.MarkerTickRecord.Reset();
 }
 
diff --git a/Engine/Source/Runtime/AnimGraphRuntime/Private/AnimNodes/AnimNode_SequenceEvaluator.cpp b/Engine/Source/Runtime/AnimGraphRuntime/Private/AnimNodes/AnimNode_SequenceEvaluator.cpp
--- a/Engine/Source/Runtime/AnimGraphRuntime/Private/AnimNodes/AnimNode_SequenceEvaluator.cpp
+++ b/Engine/Source/Runtime/AnimGraphRuntime/Private/AnimNodes/AnimNode_SequenceEvaluator.cpp
@@ -39,8 +39,10 @@ void FAnimNode_SequenceEvaluatorBase::UpdateAssetPlayer(const FAnimationUpdateCo
 		// Clamp input to a valid position on this sequence's time line.
 		CurrentExplicitTime = FMath::Clamp(CurrentExplicitTime, 0.f, CurrentSequence->GetPlayLength());
 
-		if ((!GetTeleportToExplicitTime() || (GetGroupName() != NAME_None)) && (Context.AnimInstanceProxy->IsSkeletonCompatible(CurrentSequence->GetSkeleton())))
+		if ((!GetTeleportToExplicitTime() || (GetGroupName() != NAME_None) || (GetGroupMethod() == EAnimSyncMethod::Graph)) && (Context.AnimInstanceProxy->IsSkeletonCompatible(CurrentSequence->GetSkeleton())))
 		{
+			DeltaTimeRecord.Previous = InternalTimeAccumulator;
+
 			if (bReinitialized)
 			{
 				switch (GetReinitializationBehavior())
@@ -52,22 +54,7 @@ void FAnimNode_SequenceEvaluatorBase::UpdateAssetPlayer(const FAnimationUpdateCo
 				InternalTimeAccumulator = FMath::Clamp(InternalTimeAccumulator, 0.f, CurrentSequence->GetPlayLength());
 			}
 
-			float TimeJump = CurrentExplicitTime - InternalTimeAccumulator;
-			const bool bCurrentShouldLoop = GetShouldLoop();
-			if (bCurrentShouldLoop)
-			{
-				if (FMath::Abs(TimeJump) > (CurrentSequence->GetPlayLength() * 0.5f))
-				{
-					if (TimeJump > 0.f)
-					{
-						TimeJump -= CurrentSequence->GetPlayLength();
-					}
-					else
-					{
-						TimeJump += CurrentSequence->GetPlayLength();
-					}
-				}
-			}
+			const float TimeJump = GetEffectiveDeltaTime(CurrentExplicitTime, InternalTimeAccumulator);
 
 			// if you jump from front to end or end to front, your time jump is 0.f, so nothing moves
 			// to prevent that from happening, we set current accumulator to explicit time
@@ -75,11 +62,13 @@ void FAnimNode_SequenceEvaluatorBase::UpdateAssetPlayer(const FAnimationUpdateCo
 			{
 				InternalTimeAccumulator = CurrentExplicitTime;
 			}
-			
+
+			DeltaTimeRecord.Delta = TimeJump;
+
 			const float DeltaTime = Context.GetDeltaTime();
 			const float RateScale = CurrentSequence->RateScale;
 			const float PlayRate = FMath::IsNearlyZero(DeltaTime) || FMath::IsNearlyZero(RateScale) ? 0.f : (TimeJump / (DeltaTime * RateScale));
-			CreateTickRecordForNode(Context, CurrentSequence, bCurrentShouldLoop, PlayRate);
+			CreateTickRecordForNode(Context, CurrentSequence, GetShouldLoop(), PlayRate);
 		}
 		else
 		{
@@ -103,7 +92,7 @@ void FAnimNode_SequenceEvaluatorBase::Evaluate_AnyThread(FPoseContext& Output)
 	if ((CurrentSequence != nullptr) && (Output.AnimInstanceProxy->IsSkeletonCompatible(CurrentSequence->GetSkeleton())))
 	{
 		FAnimationPoseData AnimationPoseData(Output);
-		CurrentSequence->GetAnimationPose(AnimationPoseData, FAnimExtractContext(InternalTimeAccumulator, Output.AnimInstanceProxy->ShouldExtractRootMotion()));
+		CurrentSequence->GetAnimationPose(AnimationPoseData, FAnimExtractContext(InternalTimeAccumulator, Output.AnimInstanceProxy->ShouldExtractRootMotion(), DeltaTimeRecord, GetShouldLoop()));
 	}
 	else
 	{
@@ -120,6 +109,28 @@ void FAnimNode_SequenceEvaluatorBase::GatherDebugData(FNodeDebugData& DebugData)
 	DebugData.AddDebugItem(DebugLine, true);
 }
 
+float FAnimNode_SequenceEvaluatorBase::GetEffectiveDeltaTime(float ExplicitTime, float PrevExplicitTime) const
+{
+	float DeltaTime = ExplicitTime - PrevExplicitTime;
+
+	if (GetShouldLoop())
+	{
+		if (FMath::Abs(DeltaTime) > (GetSequence()->GetPlayLength() * 0.5f))
+		{
+			if (DeltaTime > 0.f)
+			{
+				DeltaTime -= GetSequence()->GetPlayLength();
+			}
+			else
+			{
+				DeltaTime += GetSequence()->GetPlayLength();
+			}
+		}
+	}
+
+	return DeltaTime;
+}
+
 void FAnimNode_SequenceEvaluator::SetSequence(UAnimSequenceBase* InSequence)
 {
 #if WITH_EDITORONLY_DATA
diff --git a/Engine/Source/Runtime/AnimGraphRuntime/Public/AnimNodes/AnimNode_PoseHandler.h b/Engine/Source/Runtime/AnimGraphRuntime/Public/AnimNodes/AnimNode_PoseHandler.h
--- a/Engine/Source/Runtime/AnimGraphRuntime/Public/AnimNodes/AnimNode_PoseHandler.h
+++ b/Engine/Source/Runtime/AnimGraphRuntime/Public/AnimNodes/AnimNode_PoseHandler.h
@@ -41,7 +41,7 @@ public:
 
 	// FAnimNode_AssetPlayerBase Interface
 	virtual float GetAccumulatedTime() const {return 0.f;}
-	virtual void SetAccumulatedTime(const float& NewTime) {}
+	virtual void SetAccumulatedTime(float NewTime) {}
 	virtual UAnimationAsset* GetAnimAsset() const {return PoseAsset;}
 	// End of FAnimNode_AssetPlayerBase Interface
 
diff --git a/Engine/Source/Runtime/AnimGraphRuntime/Public/AnimNodes/AnimNode_RandomPlayer.h b/Engine/Source/Runtime/AnimGraphRuntime/Public/AnimNodes/AnimNode_RandomPlayer.h
--- a/Engine/Source/Runtime/AnimGraphRuntime/Public/AnimNodes/AnimNode_RandomPlayer.h
+++ b/Engine/Source/Runtime/AnimGraphRuntime/Public/AnimNodes/AnimNode_RandomPlayer.h
@@ -76,8 +76,8 @@ struct FRandomAnimPlayData
 	// The time at which the animation is currently playing.
 	float CurrentPlayTime = 0.0f;
 
-	// The previous tick's update time.
-	float PreviousPlayTime = 0.0f;
+	// Delta time record for this play through
+	FDeltaTimeRecord DeltaTimeRecord;
 
 	// Calculated play rate
 	float PlayRate = 0.0f;
diff --git a/Engine/Source/Runtime/AnimGraphRuntime/Public/AnimNodes/AnimNode_SequenceEvaluator.h b/Engine/Source/Runtime/AnimGraphRuntime/Public/AnimNodes/AnimNode_SequenceEvaluator.h
--- a/Engine/Source/Runtime/AnimGraphRuntime/Public/AnimNodes/AnimNode_SequenceEvaluator.h
+++ b/Engine/Source/Runtime/AnimGraphRuntime/Public/AnimNodes/AnimNode_SequenceEvaluator.h
@@ -49,12 +49,15 @@ public:
 
 	// FAnimNode_AssetPlayerBase Interface
 	virtual float GetAccumulatedTime() const {return GetExplicitTime();}
-	virtual void SetAccumulatedTime(const float& NewTime) { SetExplicitTime(NewTime); }
+	virtual void SetAccumulatedTime(float NewTime) { SetExplicitTime(NewTime); }
 	virtual UAnimationAsset* GetAnimAsset() const { return GetSequence(); }
 	// End of FAnimNode_AssetPlayerBase Interface
 
 	void SetExplicitPreviousTime(float PreviousTime) { InternalTimeAccumulator = PreviousTime; }
 
+	// Get the effective delta time between the previous and current frame internal time
+	virtual float GetEffectiveDeltaTime(float ExplicitTime, float PrevExplicitTime) const;
+
 	// Set the animation sequence asset to evaluate
 	virtual void SetSequence(UAnimSequenceBase* InSequence) {}
 
diff --git a/Engine/Source/Runtime/AnimGraphRuntime/Public/BoneControllers/BoneControllerTypes.h b/Engine/Source/Runtime/AnimGraphRuntime/Public/BoneControllers/BoneControllerTypes.h
new file mode 100644
--- /dev/null
+++ b/Engine/Source/Runtime/AnimGraphRuntime/Public/BoneControllers/BoneControllerTypes.h
@@ -0,0 +1,19 @@
+// Copyright Epic Games, Inc. All Rights Reserved.
+
+#pragma once
+
+#include "HAL/Platform.h"
+#include "UObject/ObjectMacros.h"
+
+UENUM(BlueprintType)
+enum class EWarpingEvaluationMode : uint8
+{
+	/* Pose warping node parameters are entirely driven by the user */
+	Manual,
+	/**
+	* Pose warping nodes may participate in being graph-driven. This means some
+	* properties of the warp may be automatically configured by the accumulated root motion
+	* contribution of the animation sub-graph leading into the node
+	*/
+	Graph
+};
\ No newline at end of file
diff --git a/Engine/Source/Runtime/Engine/Classes/Animation/AnimInstance.h b/Engine/Source/Runtime/Engine/Classes/Animation/AnimInstance.h
--- a/Engine/Source/Runtime/Engine/Classes/Animation/AnimInstance.h
+++ b/Engine/Source/Runtime/Engine/Classes/Animation/AnimInstance.h
@@ -330,15 +330,17 @@ struct FMontageActiveSlotTracker
 
 struct FMontageEvaluationState
 {
-	FMontageEvaluationState(UAnimMontage* InMontage, float InPosition, bool bInIsPlaying, bool bInIsActive, const FAlphaBlend& InBlendInfo, const UBlendProfile* InActiveBlendProfile, float InBlendStartAlpha) 
+	FMontageEvaluationState(UAnimMontage* InMontage, float InPosition, FDeltaTimeRecord InDeltaTimeRecord, bool bInIsPlaying, bool bInIsActive, const FAlphaBlend& InBlendInfo, const UBlendProfile* InActiveBlendProfile, float InBlendStartAlpha)
 		: Montage(InMontage)
 		, BlendInfo(InBlendInfo)
 		, ActiveBlendProfile(InActiveBlendProfile)
 		, MontagePosition(InPosition)
+		, DeltaTimeRecord(InDeltaTimeRecord)
 		, BlendStartAlpha(InBlendStartAlpha)
 		, bIsPlaying(bInIsPlaying)
 		, bIsActive(bInIsActive)
-	{}
+	{
+	}
 
 	// The montage to evaluate
 	TWeakObjectPtr<UAnimMontage> Montage;
@@ -351,6 +353,9 @@ struct FMontageEvaluationState
 
 	// The position to evaluate this montage at
 	float MontagePosition;
+	
+	// The previous MontagePosition and delta leading into current
+	FDeltaTimeRecord DeltaTimeRecord;
 
 	// The linear alpha value where to start blending from. So not the blended value that already has been curve sampled.
 	float BlendStartAlpha;
@@ -468,7 +473,7 @@ private:
 public:
 
 	// @todo document
-	void MakeMontageTickRecord(FAnimTickRecord& TickRecord, class UAnimMontage* Montage, float CurrentPosition, float PreviousPosition, float MoveDelta, float Weight, TArray<FPassedMarker>& MarkersPassedThisTick, FMarkerTickRecord& MarkerTickRecord);
+	void MakeMontageTickRecord(FAnimTickRecord& TickRecord, class UAnimMontage* Montage, float MoveDelta, float Weight, TArray<FPassedMarker>& MarkersPassedThisTick, FMarkerTickRecord& MarkerTickRecord);
 
 	/** Get global weight in AnimGraph for this slot node.
 	* Note: this is the weight of the node, not the weight of any potential montage it is playing. */
diff --git a/Engine/Source/Runtime/Engine/Classes/Animation/AnimMontage.h b/Engine/Source/Runtime/Engine/Classes/Animation/AnimMontage.h
--- a/Engine/Source/Runtime/Engine/Classes/Animation/AnimMontage.h
+++ b/Engine/Source/Runtime/Engine/Classes/Animation/AnimMontage.h
@@ -335,6 +335,9 @@ struct ENGINE_API FAnimMontageInstance
 	UPROPERTY(transient)
 	float DefaultBlendTimeMultiplier;
 
+	// transient value of time position and delta in the last frame known
+	FDeltaTimeRecord DeltaTimeRecord;
+
 	// marker tick record
 	FMarkerTickRecord MarkerTickRecord;
 
@@ -390,12 +393,6 @@ private:
 	//                          - we spawn all notifies
 	float NotifyWeight;
 
-	// transient value of Delta Moved in the last frame known
-	float DeltaMoved;
-
-	// transient value of previous position before move
-	float PreviousPosition;
-
 	// The current start linear alpha value of the blend. This is not stored inside the FAlphaBlend struct.
 	float BlendStartAlpha;
 
@@ -516,8 +513,8 @@ public:
 	int32 GetInstanceID() const { return InstanceID; }
 	float GetPosition() const { return Position; };
 	float GetPlayRate() const { return PlayRate; }
-	float GetDeltaMoved() const { return DeltaMoved; }
-	float GetPreviousPosition() const { return PreviousPosition;  }
+	float GetDeltaMoved() const { return DeltaTimeRecord.Delta; }
+	float GetPreviousPosition() const { return DeltaTimeRecord.Previous;  }
 	float GetBlendStartAlpha() const { return BlendStartAlpha; }
 	const FAnimMontageInstance* GetMontageSyncLeader() const { return MontageSyncLeader; } 
 	const UBlendProfile* GetActiveBlendProfile() const { return ActiveBlendProfile; }
diff --git a/Engine/Source/Runtime/Engine/Classes/Animation/AnimNode_AssetPlayerBase.h b/Engine/Source/Runtime/Engine/Classes/Animation/AnimNode_AssetPlayerBase.h
--- a/Engine/Source/Runtime/Engine/Classes/Animation/AnimNode_AssetPlayerBase.h
+++ b/Engine/Source/Runtime/Engine/Classes/Animation/AnimNode_AssetPlayerBase.h
@@ -28,7 +28,7 @@ struct ENGINE_API FAnimNode_AssetPlayerBase : public FAnimNode_Base
 	virtual float GetAccumulatedTime() const;
 
 	/** Override the currently accumulated time */
-	virtual void SetAccumulatedTime(const float& NewTime);
+	virtual void SetAccumulatedTime(float NewTime);
 
 	/** Get the animation asset associated with the node, derived classes should implement this */
 	virtual UAnimationAsset* GetAnimAsset() const;
@@ -96,6 +96,9 @@ protected:
 	/** Accumulated time used to reference the asset in this node */
 	UPROPERTY(BlueprintReadWrite, Transient, Category=DoNotEdit)
 	float InternalTimeAccumulator = 0.0f;
+	
+	/** Previous frame InternalTimeAccumulator value and effective delta time leading into the current frame */
+	FDeltaTimeRecord DeltaTimeRecord;
 
 	/** Track whether we have been full weight previously. Reset when we reach 0 weight*/
 	bool bHasBeenFullWeight = false;
diff --git a/Engine/Source/Runtime/Engine/Classes/Animation/AnimationAsset.h b/Engine/Source/Runtime/Engine/Classes/Animation/AnimationAsset.h
--- a/Engine/Source/Runtime/Engine/Classes/Animation/AnimationAsset.h
+++ b/Engine/Source/Runtime/Engine/Classes/Animation/AnimationAsset.h
@@ -72,6 +72,22 @@ struct FMarkerTickRecord
 	}
 };
 
+/**
+ * Used when sampling a given animation asset, this structure will contain the previous frame's
+ * internal sample time alongside the 'effective' delta time leading into the current frame.
+ * 
+ * An 'effective' delta time represents a value that has undergone all side effects present in the
+ * corresponding asset's TickAssetPlayer call including but not limited to syncing, play rate 
+ * adjustment, looping, etc.
+ * 
+ * For montages Delta isn't always abs(CurrentPosition-PreviousPosition) because a Montage can jump or repeat or loop
+ */
+struct FDeltaTimeRecord
+{
+	float Previous = 0.f;
+	float Delta = 0.f;
+};
+
 /** Transform definition */
 USTRUCT(BlueprintType)
 struct FBlendSampleData
@@ -102,6 +118,8 @@ struct FBlendSampleData
 	UPROPERTY()
 	float SamplePlayRate;
 
+	FDeltaTimeRecord DeltaTimeRecord;
+
 	FMarkerTickRecord MarkerTickRecord;
 
 	// transient per-bone interpolation data
@@ -118,7 +136,10 @@ struct FBlendSampleData
 		, Time(0.f)
 		, PreviousTime(0.f)
 		, SamplePlayRate(0.0f)
-	{}
+		, DeltaTimeRecord()
+	{
+	}
+
 	FBlendSampleData(int32 Index)
 		: SampleDataIndex(Index)
 		, Animation(nullptr)
@@ -127,7 +148,10 @@ struct FBlendSampleData
 		, Time(0.f)
 		, PreviousTime(0.f)
 		, SamplePlayRate(0.0f)
-	{}
+		, DeltaTimeRecord()
+	{
+	}
+
 	bool operator==( const FBlendSampleData& Other ) const 
 	{
 		// if same position, it's same point
@@ -196,12 +220,18 @@ struct FPoseCurve
 /** Animation Extraction Context */
 struct FAnimExtractContext
 {
-	/** Is Root Motion being extracted? */
-	bool bExtractRootMotion;
-
 	/** Position in animation to extract pose from */
 	float CurrentTime;
 
+	/** Is root motion being extracted? */
+	bool bExtractRootMotion;
+
+	/** Delta time range required for root motion extraction **/
+	FDeltaTimeRecord DeltaTimeRecord;
+
+	/** Is the current animation asset marked as looping? **/
+	bool bLooping;
+
 	/** 
 	 * Pose Curve Values to extract pose from pose assets. 
 	 * This is used by pose asset extraction 
@@ -215,21 +245,13 @@ struct FAnimExtractContext
 	 */
 	TArray<bool> BonesRequired;
 
-	FAnimExtractContext()
-		: bExtractRootMotion(false)
-		, CurrentTime(0.f)
-	{
-	}
-
-	FAnimExtractContext(float InCurrentTime)
-		: bExtractRootMotion(false)
-		, CurrentTime(InCurrentTime)
-	{
-	}
-
-	FAnimExtractContext(float InCurrentTime, bool InbExtractRootMotion)
-		: bExtractRootMotion(InbExtractRootMotion)
-		, CurrentTime(InCurrentTime)
+	FAnimExtractContext(float InCurrentTime = 0.f, bool InbExtractRootMotion = false, FDeltaTimeRecord InDeltaTimeRecord = {}, bool InbLooping = false)
+		: CurrentTime(InCurrentTime)
+		, bExtractRootMotion(InbExtractRootMotion)
+		, DeltaTimeRecord(InDeltaTimeRecord)
+		, bLooping(InbLooping)
+		, PoseCurves()
+		, BonesRequired()
 	{
 	}
 
@@ -239,6 +261,7 @@ struct FAnimExtractContext
 		{
 			return true;
 		}
+
 		return BonesRequired[BoneIndex];
 	}
 };
@@ -326,12 +349,14 @@ struct FAnimTickRecord
 		struct
 		{
 			float CurrentPosition;  // montage doesn't use accumulator, but this
-			float PreviousPosition;
-			float MoveDelta; // MoveDelta isn't always abs(CurrentPosition-PreviousPosition) because Montage can jump or repeat or loop
 			TArray<FPassedMarker>* MarkersPassedThisTick;
 		} Montage;
 	};
 
+	// Asset players (and other nodes) have ownership of their respective DeltaTimeRecord value/state,
+	// while an asset's tick update will forward the time-line through the tick record
+	FDeltaTimeRecord* DeltaTimeRecord = nullptr;
+
 	// marker sync related data
 	FMarkerTickRecord* MarkerTickRecord = nullptr;
 	bool bCanUseMarkerSync = false;
@@ -348,6 +373,7 @@ public:
 		, EffectiveBlendWeight(0.f)
 		, RootMotionWeightModifier(1.f)
 		, bLooping(false)
+		, DeltaTimeRecord(nullptr)
 		, MarkerTickRecord(nullptr)
 		, bCanUseMarkerSync(false)
 		, LeaderScore(0.f)
@@ -361,8 +387,12 @@ public:
 	ENGINE_API FAnimTickRecord(UBlendSpace* InBlendSpace, const FVector& InBlendInput, TArray<FBlendSampleData>& InBlendSampleDataCache, FBlendFilter& InBlendFilter, bool bInLooping, float InPlayRate, float InFinalBlendWeight, float& InCurrentTime, FMarkerTickRecord& InMarkerTickRecord);
 
 	// Create a tick record for a montage
+	UE_DEPRECATED(5.0, "Please use the montage FAnimTickRecord constructor which removes InPreviousPosition and InMoveDelta")
 	ENGINE_API FAnimTickRecord(UAnimMontage* InMontage, float InCurrentPosition, float InPreviousPosition, float InMoveDelta, float InWeight, TArray<FPassedMarker>& InMarkersPassedThisTick, FMarkerTickRecord& InMarkerTickRecord);
 
+	// Create a tick record for a montage
+	ENGINE_API FAnimTickRecord(UAnimMontage* InMontage, float InCurrentPosition, float InWeight, TArray<FPassedMarker>& InMarkersPassedThisTick, FMarkerTickRecord& InMarkerTickRecord);
+
 	// Create a tick record for a pose asset
 	ENGINE_API FAnimTickRecord(UPoseAsset* InPoseAsset, float InFinalBlendWeight);
 
diff --git a/Engine/Source/Runtime/Engine/Classes/Animation/AttributesContainer.h b/Engine/Source/Runtime/Engine/Classes/Animation/AttributesContainer.h
--- a/Engine/Source/Runtime/Engine/Classes/Animation/AttributesContainer.h
+++ b/Engine/Source/Runtime/Engine/Classes/Animation/AttributesContainer.h
@@ -614,6 +614,30 @@ namespace UE
 			*/
 			const TArray<TWeakObjectPtr<const UScriptStruct>, InAllocator>& GetUniqueTypes() const { return UniqueTypes; }
 
+			/*
+			* Populate out parameter OutAttributeKeyNames for all contained attribute key names
+			* 
+			* @return bool, whether or not any attribute key names were present 
+			*/
+			template<typename ArrayAllocator>
+			bool GetAllKeyNames(TArray<FName, ArrayAllocator>& OutAttributeKeyNames) const
+			{
+				if (!AttributeIdentifiers.IsEmpty())
+				{
+					for (const TArray<FAttributeId, InAllocator>& AttributeTypeEntry : AttributeIdentifiers)
+					{
+						for (const FAttributeId& AttributeId : AttributeTypeEntry)
+						{
+							OutAttributeKeyNames.AddUnique(AttributeId.GetName());
+						}
+					}
+
+					return true;
+				}
+
+				return false;
+			}
+
 		public:
 			/* Deprecated API */
 			template<typename DataType> 
diff --git a/Engine/Source/Runtime/Engine/Classes/Animation/BlendSpace.h b/Engine/Source/Runtime/Engine/Classes/Animation/BlendSpace.h
--- a/Engine/Source/Runtime/Engine/Classes/Animation/BlendSpace.h
+++ b/Engine/Source/Runtime/Engine/Classes/Animation/BlendSpace.h
@@ -456,10 +456,16 @@ public:
 	UE_DEPRECATED(4.26, "Use GetAnimationPose with other signature")
 	ENGINE_API void GetAnimationPose(TArray<FBlendSampleData>& BlendSampleDataCache, /*out*/ FCompactPose& OutPose, /*out*/ FBlendedCurve& OutCurve) const;
 	
+	UE_DEPRECATED(5.0, "Use GetAnimationPose with extraction context signature")
 	ENGINE_API void GetAnimationPose(TArray<FBlendSampleData>& BlendSampleDataCache, /*out*/ FAnimationPoseData& OutAnimationPoseData) const;
 
+	UE_DEPRECATED(5.0, "Use GetAnimationPose with extraction context signature")
 	ENGINE_API void GetAnimationPose(TArray<FBlendSampleData>& BlendSampleDataCache, TArrayView<FPoseLink> InPoseLinks, /*out*/ FPoseContext& Output) const;
 
+	ENGINE_API void GetAnimationPose(TArray<FBlendSampleData>& BlendSampleDataCache, const FAnimExtractContext& ExtractionContext, /*out*/ FAnimationPoseData& OutAnimationPoseData) const;
+
+	ENGINE_API void GetAnimationPose(TArray<FBlendSampleData>& BlendSampleDataCache, TArrayView<FPoseLink> InPoseLinks, const FAnimExtractContext& ExtractionContext, /*out*/ FPoseContext& Output) const;
+
 	/** Accessor for blend parameter **/
 	ENGINE_API const FBlendParameter& GetBlendParameter(const int32 Index) const;
 
@@ -630,7 +636,7 @@ protected:
 	
 private:
 	// Internal helper function for GetAnimationPose variants
-	void GetAnimationPose_Internal(TArray<FBlendSampleData>& BlendSampleDataCache, TArrayView<FPoseLink> InPoseLinks, FAnimInstanceProxy* InProxy, bool bInExpectsAdditivePose, /*out*/ FAnimationPoseData& OutAnimationPoseData) const;
+	void GetAnimationPose_Internal(TArray<FBlendSampleData>& BlendSampleDataCache, TArrayView<FPoseLink> InPoseLinks, FAnimInstanceProxy* InProxy, bool bInExpectsAdditivePose, const FAnimExtractContext& ExtractionContext, /*out*/ FAnimationPoseData& OutAnimationPoseData) const;
 
 	// Internal helper function for UpdateBlendSamples and TickAssetPlayer
 	bool UpdateBlendSamples_Internal(const FVector& InBlendSpacePosition, float InDeltaTime, TArray<FBlendSampleData>& InOutOldSampleDataList, TArray<FBlendSampleData>& InOutSampleDataCache, int32& InOutCachedTriangulationIndex) const;
diff --git a/Engine/Source/Runtime/Engine/Private/Animation/AnimCompositeBase.cpp b/Engine/Source/Runtime/Engine/Private/Animation/AnimCompositeBase.cpp
--- a/Engine/Source/Runtime/Engine/Private/Animation/AnimCompositeBase.cpp
+++ b/Engine/Source/Runtime/Engine/Private/Animation/AnimCompositeBase.cpp
@@ -569,14 +569,13 @@ void FAnimTrack::GetAnimationPose(FAnimationPoseData& OutAnimationPoseData, cons
 	{
 		if (AnimSegment->bValid)
 		{
-			float PositionInAnim = 0.f;
-			if (const UAnimSequenceBase* const AnimRef = AnimSegment->GetAnimationData(ClampedTime, PositionInAnim))
+			// Copy passed in Extraction Context, but override position and root motion parameters.
+			FAnimExtractContext SequenceExtractionContext(ExtractionContext);
+			if (const UAnimSequenceBase* const AnimRef = AnimSegment->GetAnimationData(ClampedTime, SequenceExtractionContext.CurrentTime))
 			{
-				// Copy passed in Extraction Context, but override position and root motion parameters.
-				FAnimExtractContext SequenceExtractionContext(ExtractionContext);
-				SequenceExtractionContext.CurrentTime = PositionInAnim;
+				SequenceExtractionContext.DeltaTimeRecord.Previous = SequenceExtractionContext.CurrentTime - SequenceExtractionContext.DeltaTimeRecord.Delta;
 				SequenceExtractionContext.bExtractRootMotion &= AnimRef->HasRootMotion();
-
+				SequenceExtractionContext.bLooping = AnimSegment->LoopingCount > 1;
 				AnimRef->GetAnimationPose(OutAnimationPoseData, SequenceExtractionContext);
 				bExtractedPose = true;
 			}
@@ -775,4 +774,4 @@ void UAnimCompositeBase::PostLoad()
 	Super::PostLoad();
 
 	InvalidateRecursiveAsset();
-}
+}
\ No newline at end of file
diff --git a/Engine/Source/Runtime/Engine/Private/Animation/AnimInstance.cpp b/Engine/Source/Runtime/Engine/Private/Animation/AnimInstance.cpp
--- a/Engine/Source/Runtime/Engine/Private/Animation/AnimInstance.cpp
+++ b/Engine/Source/Runtime/Engine/Private/Animation/AnimInstance.cpp
@@ -117,12 +117,10 @@ UAnimInstance::UAnimInstance(const FObjectInitializer& ObjectInitializer)
 }
 
 // this is only used by montage marker based sync
-void UAnimInstance::MakeMontageTickRecord(FAnimTickRecord& TickRecord, class UAnimMontage* Montage, float CurrentPosition, float PreviousPosition, float MoveDelta, float Weight, TArray<FPassedMarker>& MarkersPassedThisTick, FMarkerTickRecord& MarkerTickRecord)
+void UAnimInstance::MakeMontageTickRecord(FAnimTickRecord& TickRecord, class UAnimMontage* Montage, float CurrentPosition, float Weight, TArray<FPassedMarker>& MarkersPassedThisTick, FMarkerTickRecord& MarkerTickRecord)
 {
 	TickRecord.SourceAsset = Montage;
 	TickRecord.Montage.CurrentPosition = CurrentPosition;
-	TickRecord.Montage.PreviousPosition = PreviousPosition;
-	TickRecord.Montage.MoveDelta = MoveDelta;
 	TickRecord.Montage.MarkersPassedThisTick = &MarkersPassedThisTick;
 	TickRecord.MarkerTickRecord = &MarkerTickRecord;
 	TickRecord.PlayRateMultiplier = 1.f; // we don't care here, this is alreayd applied in the montageinstance::Advance
@@ -373,9 +371,15 @@ void UAnimInstance::UpdateMontageSyncGroup()
 			if (ensure(GroupNameToUse != NAME_None))
 			{
 				bRecordNeedsResetting = false;
-				FAnimTickRecord TickRecord(MontageInstance->Montage, MontageInstance->GetPosition(),
-					MontageInstance->GetPreviousPosition(), MontageInstance->GetDeltaMoved(), MontageInstance->GetWeight(),
-					MontageInstance->MarkersPassedThisTick, MontageInstance->MarkerTickRecord);
+				FAnimTickRecord TickRecord(
+					MontageInstance->Montage,
+					MontageInstance->GetPosition(),
+					MontageInstance->GetWeight(),
+					MontageInstance->MarkersPassedThisTick,
+					MontageInstance->MarkerTickRecord
+				);
+				TickRecord.DeltaTimeRecord = &MontageInstance->DeltaTimeRecord;
+
 				UE::Anim::FAnimSyncParams Params(GroupNameToUse);
 				GetProxyOnGameThread<FAnimInstanceProxy>().AddTickRecord(TickRecord, Params);
 
@@ -3324,6 +3328,7 @@ void UAnimInstance::UpdateMontageEvaluationData()
 				(
 					MontageInstance->Montage,
 					MontageInstance->GetPosition(),
+					MontageInstance->DeltaTimeRecord,
 					MontageInstance->bPlaying,
 					MontageInstance->IsActive(),
 					MontageInstance->GetBlend(),
diff --git a/Engine/Source/Runtime/Engine/Private/Animation/AnimInstanceProxy.cpp b/Engine/Source/Runtime/Engine/Private/Animation/AnimInstanceProxy.cpp
--- a/Engine/Source/Runtime/Engine/Private/Animation/AnimInstanceProxy.cpp
+++ b/Engine/Source/Runtime/Engine/Private/Animation/AnimInstanceProxy.cpp
@@ -689,6 +689,9 @@ PRAGMA_DISABLE_DEPRECATION_WARNINGS
 	MakeSequenceTickRecord(TickRecord, Sequence, bLooping, PlayRate, /*FinalBlendWeight=*/ 1.0f, CurrentTime, MarkerTickRecord);
 PRAGMA_ENABLE_DEPRECATION_WARNINGS
 
+	FDeltaTimeRecord DeltaTimeRecord;
+	TickRecord.DeltaTimeRecord = &DeltaTimeRecord;
+
 	FAnimAssetTickContext TickContext(DeltaSeconds, RootMotionMode, true);
 	TickRecord.SourceAsset->TickAssetPlayer(TickRecord, NotifyQueue, TickContext);
 }
@@ -700,6 +703,9 @@ PRAGMA_DISABLE_DEPRECATION_WARNINGS
 	MakeBlendSpaceTickRecord(TickRecord, BlendSpace, BlendInput, BlendSampleDataCache, BlendFilter, bLooping, PlayRate, /*FinalBlendWeight=*/ 1.0f, CurrentTime, MarkerTickRecord);
 PRAGMA_ENABLE_DEPRECATION_WARNINGS
 
+	FDeltaTimeRecord DeltaTimeRecord;
+	TickRecord.DeltaTimeRecord = &DeltaTimeRecord;
+
 	FAnimAssetTickContext TickContext(DeltaSeconds, RootMotionMode, true);
 	TickRecord.SourceAsset->TickAssetPlayer(TickRecord, NotifyQueue, TickContext);
 }
@@ -1359,7 +1365,7 @@ void FAnimInstanceProxy::SlotEvaluatePoseWithBlendProfiles(const FName& SlotNode
 			NewPose.Curve.InitFrom(RequiredBones);
 
 			// Extract pose from Track.
-			FAnimExtractContext ExtractionContext(EvalState.MontagePosition, Montage->HasRootMotion() && RootMotionMode != ERootMotionMode::NoRootMotionExtraction);
+			FAnimExtractContext ExtractionContext(EvalState.MontagePosition, Montage->HasRootMotion() && RootMotionMode != ERootMotionMode::NoRootMotionExtraction, EvalState.DeltaTimeRecord);
 			FAnimationPoseData NewAnimationPoseData(NewPose);
 			AnimTrack->GetAnimationPose(NewAnimationPoseData, ExtractionContext);
 
@@ -1698,7 +1704,7 @@ void FAnimInstanceProxy::SlotEvaluatePose(const FName& SlotNodeName, const FAnim
 			NewPose.Curve.InitFrom(RequiredBones);
 
 			// Extract pose from Track
-			FAnimExtractContext ExtractionContext(EvalState.MontagePosition, Montage->HasRootMotion() && RootMotionMode != ERootMotionMode::NoRootMotionExtraction);
+			FAnimExtractContext ExtractionContext(EvalState.MontagePosition, Montage->HasRootMotion() && RootMotionMode != ERootMotionMode::NoRootMotionExtraction, EvalState.DeltaTimeRecord);
 
 			FAnimationPoseData NewAnimationPoseData(NewPose);
 			AnimTrack->GetAnimationPose(NewAnimationPoseData, ExtractionContext);
diff --git a/Engine/Source/Runtime/Engine/Private/Animation/AnimMontage.cpp b/Engine/Source/Runtime/Engine/Private/Animation/AnimMontage.cpp
--- a/Engine/Source/Runtime/Engine/Private/Animation/AnimMontage.cpp
+++ b/Engine/Source/Runtime/Engine/Private/Animation/AnimMontage.cpp
@@ -1156,9 +1156,10 @@ void UAnimMontage::TickAssetPlayer(FAnimTickRecord& Instance, struct FAnimNotify
 	//if (ensure (Context.IsLeader()))
 	if ((Context.IsLeader()))
 	{
+		check(Instance.DeltaTimeRecord);
 		const float CurrentTime = Instance.Montage.CurrentPosition;
-		const float PreviousTime = Instance.Montage.PreviousPosition;
-		const float MoveDelta = Instance.Montage.MoveDelta;
+		const float PreviousTime = Instance.DeltaTimeRecord->Previous;
+		const float MoveDelta = Instance.DeltaTimeRecord->Delta;
 
 		Context.SetLeaderDelta(MoveDelta);
 		Context.SetPreviousAnimationPositionRatio(PreviousTime / GetPlayLength());
@@ -1324,8 +1325,6 @@ FAnimMontageInstance::FAnimMontageInstance()
 	, bInterrupted(false)
 	, PreviousWeight(0.f)
 	, NotifyWeight(0.f)
-	, DeltaMoved(0.f)
-	, PreviousPosition(0.f)
 	, BlendStartAlpha(0.0f)
 	, SyncGroupName(NAME_None)
 	, ActiveBlendProfile(nullptr)
@@ -1349,8 +1348,6 @@ FAnimMontageInstance::FAnimMontageInstance(UAnimInstance * InAnimInstance)
 	, bInterrupted(false)
 	, PreviousWeight(0.f)
 	, NotifyWeight(0.f)
-	, DeltaMoved(0.f)
-	, PreviousPosition(0.f)
 	, BlendStartAlpha(0.0f)
 	, SyncGroupName(NAME_None)
 	, ActiveBlendProfile(nullptr)
@@ -2311,8 +2308,8 @@ void FAnimMontageInstance::Advance(float DeltaTime, struct FRootMotionMovementPa
 		{
 			const bool bExtractRootMotion = (OutRootMotionParams != nullptr) && Montage->HasRootMotion();
 			
-			DeltaMoved = 0.f;
-			PreviousPosition = Position;
+			DeltaTimeRecord.Delta = 0.f;
+			DeltaTimeRecord.Previous = Position;
 
 			bDidUseMarkerSyncThisTick = CanUseMarkerSync();
 			if (bDidUseMarkerSyncThisTick)
@@ -2353,7 +2350,7 @@ void FAnimMontageInstance::Advance(float DeltaTime, struct FRootMotionMovementPa
 				}
 
 				const float SubStepDeltaMove = MontageSubStepper.GetDeltaMove();
-				DeltaMoved += SubStepDeltaMove;
+				DeltaTimeRecord.Delta += SubStepDeltaMove;
 				const bool bPlayingForward = MontageSubStepper.GetbPlayingForward();
 
 				// If current section is last one, check to trigger a blend out and if it hasn't stopped yet, see if we should stop
diff --git a/Engine/Source/Runtime/Engine/Private/Animation/AnimNodeBase.cpp b/Engine/Source/Runtime/Engine/Private/Animation/AnimNodeBase.cpp
--- a/Engine/Source/Runtime/Engine/Private/Animation/AnimNodeBase.cpp
+++ b/Engine/Source/Runtime/Engine/Private/Animation/AnimNodeBase.cpp
@@ -7,6 +7,7 @@
 #include "UObject/CoreObjectVersion.h"
 #include "Animation/AnimAttributes.h"
 #include "Animation/AnimSubsystem_Base.h"
+#include "Animation/AnimRootMotionProvider.h"
 
 /////////////////////////////////////////////////////
 // FAnimationBaseContext
@@ -364,6 +365,15 @@ void FPoseLink::Evaluate(FPoseContext& Output)
 			if(Output.CustomAttributes.ContainsData())
 			{
 				Output.AnimInstanceProxy->RecordNodeAttribute(*Output.AnimInstanceProxy, SourceID, LinkID, UE::Anim::FAttributes::Attributes);
+
+				TArray<FName, TInlineAllocator<8>> AttributeKeyNames;
+				if (Output.CustomAttributes.GetAllKeyNames(AttributeKeyNames))
+				{
+					for (const FName& AttributeKeyName : AttributeKeyNames)
+					{
+						Output.AnimInstanceProxy->RecordNodeAttribute(*Output.AnimInstanceProxy, SourceID, LinkID, AttributeKeyName);
+					}
+				}
 			}
 			if(Output.Curve.NumValid() > 0)
 			{
@@ -465,6 +475,15 @@ void FComponentSpacePoseLink::EvaluateComponentSpace(FComponentSpacePoseContext&
 			if(Output.CustomAttributes.ContainsData())
 			{
 				Output.AnimInstanceProxy->RecordNodeAttribute(*Output.AnimInstanceProxy, SourceID, LinkID, UE::Anim::FAttributes::Attributes);
+
+				TArray<FName, FAnimStackAllocator> AttributeKeyNames;
+				if (Output.CustomAttributes.GetAllKeyNames(AttributeKeyNames))
+				{
+					for (const FName& AttributeKeyName : AttributeKeyNames)
+					{
+						Output.AnimInstanceProxy->RecordNodeAttribute(*Output.AnimInstanceProxy, SourceID, LinkID, AttributeKeyName);
+					}
+				}
 			}
 			if(Output.Curve.NumValid() > 0)
 			{
diff --git a/Engine/Source/Runtime/Engine/Private/Animation/AnimNode_AssetPlayerBase.cpp b/Engine/Source/Runtime/Engine/Private/Animation/AnimNode_AssetPlayerBase.cpp
--- a/Engine/Source/Runtime/Engine/Private/Animation/AnimNode_AssetPlayerBase.cpp
+++ b/Engine/Source/Runtime/Engine/Private/Animation/AnimNode_AssetPlayerBase.cpp
@@ -45,6 +45,7 @@ void FAnimNode_AssetPlayerBase::CreateTickRecordForNode(const FAnimationUpdateCo
 		Context.GetSharedContext()->MessageStack.MakeEventContextData(TickRecord.ContextData);
 	}
 	TickRecord.RootMotionWeightModifier = Context.GetRootMotionWeightModifier();
+	TickRecord.DeltaTimeRecord = &DeltaTimeRecord;
 
 	SyncScope.AddTickRecord(TickRecord, SyncParams, UE::Anim::FAnimSyncDebugInfo(Context));
 
@@ -61,7 +62,7 @@ float FAnimNode_AssetPlayerBase::GetAccumulatedTime() const
 	return InternalTimeAccumulator;
 }
 
-void FAnimNode_AssetPlayerBase::SetAccumulatedTime(const float& NewTime)
+void FAnimNode_AssetPlayerBase::SetAccumulatedTime(float NewTime)
 {
 	InternalTimeAccumulator = NewTime;
 }
diff --git a/Engine/Source/Runtime/Engine/Private/Animation/AnimNode_SequencePlayer.cpp b/Engine/Source/Runtime/Engine/Private/Animation/AnimNode_SequencePlayer.cpp
--- a/Engine/Source/Runtime/Engine/Private/Animation/AnimNode_SequencePlayer.cpp
+++ b/Engine/Source/Runtime/Engine/Private/Animation/AnimNode_SequencePlayer.cpp
@@ -125,7 +125,7 @@ void FAnimNode_SequencePlayerBase::Evaluate_AnyThread(FPoseContext& Output)
 		}
 
 		FAnimationPoseData AnimationPoseData(Output);
-		CurrentSequence->GetAnimationPose(AnimationPoseData, FAnimExtractContext(InternalTimeAccumulator, Output.AnimInstanceProxy->ShouldExtractRootMotion()));
+		CurrentSequence->GetAnimationPose(AnimationPoseData, FAnimExtractContext(InternalTimeAccumulator, Output.AnimInstanceProxy->ShouldExtractRootMotion(), DeltaTimeRecord, GetLoopAnimation()));
 	}
 	else
 	{
diff --git a/Engine/Source/Runtime/Engine/Private/Animation/AnimRootMotionProvider.cpp b/Engine/Source/Runtime/Engine/Private/Animation/AnimRootMotionProvider.cpp
new file mode 100644
--- /dev/null
+++ b/Engine/Source/Runtime/Engine/Private/Animation/AnimRootMotionProvider.cpp
@@ -0,0 +1,27 @@
+// Copyright Epic Games, Inc. All Rights Reserved.
+
+#include "Animation/AnimRootMotionProvider.h"
+#include "Features/IModularFeatures.h"
+
+
+namespace UE { namespace Anim {
+
+const FName IAnimRootMotionProvider::ModularFeatureName(TEXT("AnimationWarping"));
+
+const FName IAnimRootMotionProvider::RootMotionDeltaAttributeName(TEXT("RootMotionDelta"));
+
+const IAnimRootMotionProvider* IAnimRootMotionProvider::Get()
+{
+	if (IsAvailable())
+	{
+		return &IModularFeatures::Get().GetModularFeature<const IAnimRootMotionProvider>(ModularFeatureName);
+	}
+	return nullptr;
+}
+
+bool IAnimRootMotionProvider::IsAvailable()
+{
+	return IModularFeatures::Get().IsModularFeatureAvailable(ModularFeatureName);
+}
+
+}}; // namespace UE::Anim
\ No newline at end of file
diff --git a/Engine/Source/Runtime/Engine/Private/Animation/AnimSequence.cpp b/Engine/Source/Runtime/Engine/Private/Animation/AnimSequence.cpp
--- a/Engine/Source/Runtime/Engine/Private/Animation/AnimSequence.cpp
+++ b/Engine/Source/Runtime/Engine/Private/Animation/AnimSequence.cpp
@@ -23,6 +23,7 @@
 #include "AnimationRuntime.h"
 #include "Animation/AnimCompress.h"
 #include "Components/SkeletalMeshComponent.h"
+#include "Animation/AnimRootMotionProvider.h"
 #include "Animation/AnimNotifies/AnimNotify.h"
 #include "Animation/BlendSpace.h"
 #include "Animation/Rig.h"
@@ -1393,6 +1394,15 @@ void UAnimSequence::GetAnimationPose(FAnimationPoseData& OutAnimationPoseData, c
 		GetBonePose(OutAnimationPoseData, ExtractionContext);
 	}
 
+	// If the sequence has root motion enabled, allow sampling of a root motion delta into the custom attribute container of the outgoing pose
+	if (HasRootMotion())
+	{
+		if (const UE::Anim::IAnimRootMotionProvider* RootMotionProvider = UE::Anim::IAnimRootMotionProvider::Get())
+		{
+			RootMotionProvider->SampleRootMotion(ExtractionContext.DeltaTimeRecord, *this, ExtractionContext.bLooping, OutAnimationPoseData.GetAttributes());
+		}
+	}
+
 	// Check that all bone atoms coming from animation are normalized
 #if DO_CHECK && WITH_EDITORONLY_DATA
 	check(OutPose.IsNormalized());
diff --git a/Engine/Source/Runtime/Engine/Private/Animation/AnimSequenceBase.cpp b/Engine/Source/Runtime/Engine/Private/Animation/AnimSequenceBase.cpp
--- a/Engine/Source/Runtime/Engine/Private/Animation/AnimSequenceBase.cpp
+++ b/Engine/Source/Runtime/Engine/Private/Animation/AnimSequenceBase.cpp
@@ -424,31 +424,30 @@ void UAnimSequenceBase::RemapTracksToNewSkeleton(USkeleton* NewSkeleton, bool bC
 
 void UAnimSequenceBase::TickAssetPlayer(FAnimTickRecord& Instance, struct FAnimNotifyQueue& NotifyQueue, FAnimAssetTickContext& Context) const
 {
-	float& CurrentTime = *(Instance.TimeAccumulator);
+	float CurrentTime = *(Instance.TimeAccumulator);
 	float PreviousTime = CurrentTime;
-	const float PlayRate = Instance.PlayRateMultiplier * this->RateScale;
+	float DeltaTime = 0.f;
 
-	float MoveDelta = 0.f;
+	const float PlayRate = Instance.PlayRateMultiplier * this->RateScale;
 
 	if( Context.IsLeader() )
 	{
-		const float DeltaTime = Context.GetDeltaTime();
-		MoveDelta = PlayRate * DeltaTime;
+		DeltaTime = PlayRate * Context.GetDeltaTime();
 
-		Context.SetLeaderDelta(MoveDelta);
+		Context.SetLeaderDelta(DeltaTime);
 		Context.SetPreviousAnimationPositionRatio(PreviousTime / GetPlayLength());
 
-		if (MoveDelta != 0.f)
+		if (DeltaTime != 0.f)
 		{
 			if (Instance.bCanUseMarkerSync && Context.CanUseMarkerPosition())
 			{
-				TickByMarkerAsLeader(*Instance.MarkerTickRecord, Context.MarkerTickContext, CurrentTime, PreviousTime, MoveDelta, Instance.bLooping, Instance.MirrorDataTable);
+				TickByMarkerAsLeader(*Instance.MarkerTickRecord, Context.MarkerTickContext, CurrentTime, PreviousTime, DeltaTime, Instance.bLooping, Instance.MirrorDataTable);
 			}
 			else
 			{
 				// Advance time
-				FAnimationRuntime::AdvanceTime(Instance.bLooping, MoveDelta, CurrentTime, GetPlayLength());
-				UE_LOG(LogAnimMarkerSync, Log, TEXT("Leader (%s) (normal advance)  - PreviousTime (%0.2f), CurrentTime (%0.2f), MoveDelta (%0.2f), Looping (%d) "), *GetName(), PreviousTime, CurrentTime, MoveDelta, Instance.bLooping ? 1 : 0);
+				FAnimationRuntime::AdvanceTime(Instance.bLooping, DeltaTime, CurrentTime, GetPlayLength());
+				UE_LOG(LogAnimMarkerSync, Log, TEXT("Leader (%s) (normal advance)  - PreviousTime (%0.2f), CurrentTime (%0.2f), MoveDelta (%0.2f), Looping (%d) "), *GetName(), PreviousTime, CurrentTime, DeltaTime, Instance.bLooping ? 1 : 0);
 			}
 		}
 
@@ -465,18 +464,17 @@ void UAnimSequenceBase::TickAssetPlayer(FAnimTickRecord& Instance, struct FAnimN
 			}
 			else
 			{
-				const float DeltaTime = Context.GetDeltaTime();
-				const float MyMoveDelta = PlayRate * DeltaTime;
+				DeltaTime = PlayRate * Context.GetDeltaTime();
 				// If leader is not valid, advance time as normal, do not jump position and pop.
-				FAnimationRuntime::AdvanceTime(Instance.bLooping, MyMoveDelta, CurrentTime, GetPlayLength());
-				UE_LOG(LogAnimMarkerSync, Log, TEXT("Follower (%s) (normal advance)  - PreviousTime (%0.2f), CurrentTime (%0.2f), MoveDelta (%0.2f), Looping (%d) "), *GetName(), PreviousTime, CurrentTime, MoveDelta, Instance.bLooping ? 1 : 0);
+				FAnimationRuntime::AdvanceTime(Instance.bLooping, DeltaTime, CurrentTime, GetPlayLength());
+				UE_LOG(LogAnimMarkerSync, Log, TEXT("Follower (%s) (normal advance) - PreviousTime (%0.2f), CurrentTime (%0.2f), MoveDelta (%0.2f), Looping (%d) "), *GetName(), PreviousTime, CurrentTime, DeltaTime, Instance.bLooping ? 1 : 0);
 			}
 		}
 		else
 		{
 			PreviousTime = Context.GetPreviousAnimationPositionRatio() * GetPlayLength();
 			CurrentTime = Context.GetAnimationPositionRatio() * GetPlayLength();
-			UE_LOG(LogAnimMarkerSync, Log, TEXT("Follower (%s) (normalized position advance) - PreviousTime (%0.2f), CurrentTime (%0.2f), MoveDelta (%0.2f), Looping (%d) "), *GetName(), PreviousTime, CurrentTime, MoveDelta, Instance.bLooping ? 1 : 0);
+			UE_LOG(LogAnimMarkerSync, Log, TEXT("Follower (%s) (normalized position advance) - PreviousTime (%0.2f), CurrentTime (%0.2f), MoveDelta (%0.2f), Looping (%d) "), *GetName(), PreviousTime, CurrentTime, DeltaTime, Instance.bLooping ? 1 : 0);
 		}
 
 		//@TODO: NOTIFIES: Calculate AdvanceType based on what the new delta time is
@@ -484,16 +482,24 @@ void UAnimSequenceBase::TickAssetPlayer(FAnimTickRecord& Instance, struct FAnimN
 		if( CurrentTime != PreviousTime )
 		{
 			// Figure out delta time 
-			MoveDelta = CurrentTime - PreviousTime;
+			DeltaTime = CurrentTime - PreviousTime;
 			// if we went against play rate, then loop around.
-			if( (MoveDelta * PlayRate) < 0.f )
+			if( (DeltaTime * PlayRate) < 0.f )
 			{
-				MoveDelta += FMath::Sign<float>(PlayRate) * GetPlayLength();
+				DeltaTime += FMath::Sign<float>(PlayRate) * GetPlayLength();
 			}
 		}
 	}
 
-	HandleAssetPlayerTickedInternal(Context, PreviousTime, MoveDelta, Instance, NotifyQueue);
+	// Assign the instance's TimeAccumulator after all side effects on CurrentTime have been applied
+	*(Instance.TimeAccumulator) = CurrentTime;
+
+	// Capture the final adjusted delta time and previous frame time as an asset player record
+	check(Instance.DeltaTimeRecord);
+	Instance.DeltaTimeRecord->Previous = PreviousTime;
+	Instance.DeltaTimeRecord->Delta = DeltaTime;
+
+	HandleAssetPlayerTickedInternal(Context, PreviousTime, DeltaTime, Instance, NotifyQueue);
 }
 
 void UAnimSequenceBase::TickByMarkerAsFollower(FMarkerTickRecord &Instance, FMarkerTickContext &MarkerContext, float& CurrentTime, float& OutPreviousTime, const float MoveDelta, const bool bLooping, const UMirrorDataTable* MirrorTable) const
@@ -939,7 +945,7 @@ void UAnimSequenceBase::Serialize(FArchive& Ar)
 	PRAGMA_ENABLE_DEPRECATION_WARNINGS
 }
 
-void UAnimSequenceBase::GetAnimationPose(struct FCompactPose& OutPose, FBlendedCurve & OutCurve, const FAnimExtractContext & ExtractionContext) const
+void UAnimSequenceBase::GetAnimationPose(struct FCompactPose& OutPose, FBlendedCurve & OutCurve, const FAnimExtractContext& ExtractionContext) const
 {
 	UE::Anim::FStackAttributeContainer TempAttributes;
 	FAnimationPoseData OutAnimationPoseData(OutPose, OutCurve, TempAttributes);
diff --git a/Engine/Source/Runtime/Engine/Private/Animation/AnimSingleNodeInstanceProxy.cpp b/Engine/Source/Runtime/Engine/Private/Animation/AnimSingleNodeInstanceProxy.cpp
--- a/Engine/Source/Runtime/Engine/Private/Animation/AnimSingleNodeInstanceProxy.cpp
+++ b/Engine/Source/Runtime/Engine/Private/Animation/AnimSingleNodeInstanceProxy.cpp
@@ -29,7 +29,9 @@ void FAnimSingleNodeInstanceProxy::Initialize(UAnimInstance* InAnimInstance)
 
 	// it's already doing it when evaluate
 	BlendSpacePosition = FVector::ZeroVector;
+
 	CurrentTime = 0.f;
+	DeltaTimeRecord = FDeltaTimeRecord();
 
 	// initialize node manually 
 	FAnimationInitializeContext InitContext(this);
@@ -160,6 +162,8 @@ void FAnimSingleNodeInstanceProxy::InternalBlendSpaceEvaluatePose(class UBlendSp
 {
 	FAnimationPoseData AnimationPoseData = { OutContext.Pose, OutContext.Curve, OutContext.CustomAttributes };
 
+	FAnimExtractContext ExtractionContext(CurrentTime, ShouldExtractRootMotion(), DeltaTimeRecord, bLooping);
+
 	if (BlendSpace->IsValidAdditive())
 	{
 		FCompactPose& OutPose = OutContext.Pose;
@@ -182,14 +186,14 @@ void FAnimSingleNodeInstanceProxy::InternalBlendSpaceEvaluatePose(class UBlendSp
 		}
 
 		FAnimationPoseData AdditiveAnimationPoseData = { AdditivePose, AdditiveCurve, AdditiveAttributes };
-		BlendSpace->GetAnimationPose(BlendSampleDataCache, AdditiveAnimationPoseData);
+		BlendSpace->GetAnimationPose(BlendSampleDataCache, ExtractionContext, AdditiveAnimationPoseData);
 
 		enum EAdditiveAnimationType AdditiveType = BlendSpace->bRotationBlendInMeshSpace? AAT_RotationOffsetMeshSpace : AAT_LocalSpaceBase;
 		FAnimationRuntime::AccumulateAdditivePose(AnimationPoseData, AdditiveAnimationPoseData, 1.f, AdditiveType);
 	}
 	else
 	{
-		BlendSpace->GetAnimationPose(BlendSampleDataCache, AnimationPoseData);
+		BlendSpace->GetAnimationPose(BlendSampleDataCache, ExtractionContext, AnimationPoseData);
 	}
 }
 
@@ -303,10 +307,10 @@ void FAnimNode_SingleNode::Evaluate_AnyThread(FPoseContext& Output)
 		}
 		else if (UAnimSequence* Sequence = Cast<UAnimSequence>(Proxy->CurrentAsset))
 		{
+			FAnimExtractContext ExtractionContext(Proxy->CurrentTime, Sequence->bEnableRootMotion, Proxy->DeltaTimeRecord, Proxy->bLooping);
+
 			if (Sequence->IsValidAdditive())
 			{
-				FAnimExtractContext ExtractionContext(Proxy->CurrentTime, Sequence->bEnableRootMotion);
-
 				if (bCanProcessAdditiveAnimationsLocal)
 				{
 					Sequence->GetAdditiveBasePose(OutputAnimationPoseData, ExtractionContext);
@@ -331,7 +335,7 @@ void FAnimNode_SingleNode::Evaluate_AnyThread(FPoseContext& Output)
 			else
 			{
 				// if SkeletalMesh isn't there, we'll need to use skeleton
-				Sequence->GetAnimationPose(OutputAnimationPoseData, FAnimExtractContext(Proxy->CurrentTime, Sequence->bEnableRootMotion));
+				Sequence->GetAnimationPose(OutputAnimationPoseData, ExtractionContext);
 			}
 		}
 		else if (UAnimStreamable* Streamable = Cast<UAnimStreamable>(Proxy->CurrentAsset))
@@ -362,12 +366,12 @@ void FAnimNode_SingleNode::Evaluate_AnyThread(FPoseContext& Output)
 			else*/
 			{
 				// if SkeletalMesh isn't there, we'll need to use skeleton
-				Streamable->GetAnimationPose(OutputAnimationPoseData, FAnimExtractContext(Proxy->CurrentTime, Streamable->bEnableRootMotion));
+				Streamable->GetAnimationPose(OutputAnimationPoseData, FAnimExtractContext(Proxy->CurrentTime, Streamable->bEnableRootMotion, Proxy->DeltaTimeRecord, Proxy->bLooping));
 			}
 		}
 		else if (UAnimComposite* Composite = Cast<UAnimComposite>(Proxy->CurrentAsset))
 		{
-			FAnimExtractContext ExtractionContext(Proxy->CurrentTime, Proxy->ShouldExtractRootMotion());
+			FAnimExtractContext ExtractionContext(Proxy->CurrentTime, Proxy->ShouldExtractRootMotion(), Proxy->DeltaTimeRecord, Proxy->bLooping);
 			const FAnimTrack& AnimTrack = Composite->AnimationTrack;
 
 			// find out if this is additive animation
@@ -555,6 +559,8 @@ void FAnimNode_SingleNode::Update_AnyThread(const FAnimationUpdateContext& Conte
 		if (UBlendSpace* BlendSpace = Cast<UBlendSpace>(Proxy->CurrentAsset))
 		{
 			FAnimTickRecord TickRecord(BlendSpace, Proxy->BlendSpacePosition, Proxy->BlendSampleData, Proxy->BlendFilter, Proxy->bLooping, NewPlayRate, 1.f, /*inout*/ Proxy->CurrentTime, Proxy->MarkerTickRecord);
+			TickRecord.DeltaTimeRecord = &(Proxy->DeltaTimeRecord);
+			
 			SyncScope.AddTickRecord(TickRecord);
 
 			TRACE_ANIM_TICK_RECORD(Context, TickRecord);
@@ -565,6 +571,8 @@ void FAnimNode_SingleNode::Update_AnyThread(const FAnimationUpdateContext& Conte
 		else if (UAnimSequence* Sequence = Cast<UAnimSequence>(Proxy->CurrentAsset))
 		{
 			FAnimTickRecord TickRecord(Sequence, Proxy->bLooping, NewPlayRate, 1.f, /*inout*/ Proxy->CurrentTime, Proxy->MarkerTickRecord);
+			TickRecord.DeltaTimeRecord = &(Proxy->DeltaTimeRecord);
+			
 			SyncScope.AddTickRecord(TickRecord);
 
 			TRACE_ANIM_TICK_RECORD(Context, TickRecord);
@@ -582,6 +590,8 @@ void FAnimNode_SingleNode::Update_AnyThread(const FAnimationUpdateContext& Conte
 		else if (UAnimStreamable* Streamable = Cast<UAnimStreamable>(Proxy->CurrentAsset))
 		{
 			FAnimTickRecord TickRecord(Streamable, Proxy->bLooping, NewPlayRate, 1.f, /*inout*/ Proxy->CurrentTime, Proxy->MarkerTickRecord);
+			TickRecord.DeltaTimeRecord = &(Proxy->DeltaTimeRecord);
+			
 			SyncScope.AddTickRecord(TickRecord);
 
 			TRACE_ANIM_TICK_RECORD(Context, TickRecord);
@@ -599,6 +609,8 @@ void FAnimNode_SingleNode::Update_AnyThread(const FAnimationUpdateContext& Conte
 		else if(UAnimComposite* Composite = Cast<UAnimComposite>(Proxy->CurrentAsset))
 		{
 			FAnimTickRecord TickRecord(Composite, Proxy->bLooping, NewPlayRate, 1.f, /*inout*/ Proxy->CurrentTime, Proxy->MarkerTickRecord);
+			TickRecord.DeltaTimeRecord = &(Proxy->DeltaTimeRecord);
+			
 			SyncScope.AddTickRecord(TickRecord);
 
 			TRACE_ANIM_TICK_RECORD(Context, TickRecord);
@@ -626,6 +638,7 @@ void FAnimNode_SingleNode::Update_AnyThread(const FAnimationUpdateContext& Conte
 			if (ActiveMontageEvaluationState)
 			{
 				Proxy->CurrentTime = ActiveMontageEvaluationState->MontagePosition;
+				Proxy->DeltaTimeRecord = ActiveMontageEvaluationState->DeltaTimeRecord;
 			}
 			else if (Proxy->bPlaying)
 			{
@@ -638,6 +651,8 @@ void FAnimNode_SingleNode::Update_AnyThread(const FAnimationUpdateContext& Conte
 		else if (UPoseAsset* PoseAsset = Cast<UPoseAsset>(Proxy->CurrentAsset))
 		{
 			FAnimTickRecord TickRecord(PoseAsset, 1.f);
+			TickRecord.DeltaTimeRecord = &(Proxy->DeltaTimeRecord);
+
 			SyncScope.AddTickRecord(TickRecord);
 
 			TRACE_ANIM_TICK_RECORD(Context, TickRecord);
diff --git a/Engine/Source/Runtime/Engine/Private/Animation/BlendSpace.cpp b/Engine/Source/Runtime/Engine/Private/Animation/BlendSpace.cpp
--- a/Engine/Source/Runtime/Engine/Private/Animation/BlendSpace.cpp
+++ b/Engine/Source/Runtime/Engine/Private/Animation/BlendSpace.cpp
@@ -279,7 +279,13 @@ void UBlendSpace::TickAssetPlayer(FAnimTickRecord& Instance, struct FAnimNotifyQ
 	TArray<FBlendSampleData>& SampleDataList = *Instance.BlendSpace.BlendSampleDataCache;
 
 	const float DeltaTime = Context.GetDeltaTime();
-	float MoveDelta = Instance.PlayRateMultiplier * DeltaTime;
+
+	check(Instance.DeltaTimeRecord);
+	// The blend space delta time record isn't currently used directly
+	// Instead, the blend sample time record's are used to drive pose evaluation
+	// As a consequence, we currently follow the convention that TimeAccumulator and Previous are normalized
+	Instance.DeltaTimeRecord->Previous = *(Instance.TimeAccumulator);
+	Instance.DeltaTimeRecord->Delta = Instance.PlayRateMultiplier * DeltaTime; // Referenced as MoveDelta within comments and logging
 
 	// this happens even if MoveDelta == 0.f. This still should happen if it is being interpolated
 	// since we allow setting position of blendspace, we can't ignore MoveDelta == 0.f
@@ -350,7 +356,7 @@ void UBlendSpace::TickAssetPlayer(FAnimTickRecord& Instance, struct FAnimNotifyQ
 					}
 				}
 
-				MoveDelta *= FilterMultiplier;
+				Instance.DeltaTimeRecord->Delta *= FilterMultiplier;
 				UE_LOG(LogAnimation, Log, TEXT("BlendSpace(%s) - FilteredBlendInput(%s) : FilteredBlendInput(%s), FilterMultiplier(%0.2f)"), *GetName(), *BlendSpacePosition.ToString(), *FilteredBlendInput.ToString(), FilterMultiplier);
 			}
 
@@ -378,7 +384,7 @@ void UBlendSpace::TickAssetPlayer(FAnimTickRecord& Instance, struct FAnimNotifyQ
 
 			if (PreInterpAnimLength > 0.f && NewAnimLength > 0.f)
 			{
-				MoveDelta *= PreInterpAnimLength / NewAnimLength;
+				Instance.DeltaTimeRecord->Delta *= PreInterpAnimLength / NewAnimLength;
 			}
 
 			float& NormalizedCurrentTime = *(Instance.TimeAccumulator);
@@ -444,9 +450,9 @@ void UBlendSpace::TickAssetPlayer(FAnimTickRecord& Instance, struct FAnimNotifyQ
 				{
 					// Advance time using current/new anim length
 					float CurrentTime = NormalizedCurrentTime * NewAnimLength;
-					FAnimationRuntime::AdvanceTime(Instance.bLooping, MoveDelta, /*inout*/ CurrentTime, NewAnimLength);
+					FAnimationRuntime::AdvanceTime(Instance.bLooping, Instance.DeltaTimeRecord->Delta, /*inout*/ CurrentTime, NewAnimLength);
 					NormalizedCurrentTime = NewAnimLength ? (CurrentTime / NewAnimLength) : 0.0f;
-					UE_LOG(LogAnimMarkerSync, Log, TEXT("Leader (%s) (bCanDoMarkerSync == false)  - PreviousTime (%0.2f), CurrentTime (%0.2f), MoveDelta (%0.2f) "), *GetName(), NormalizedPreviousTime, NormalizedCurrentTime, MoveDelta);
+					UE_LOG(LogAnimMarkerSync, Log, TEXT("Leader (%s) (bCanDoMarkerSync == false)  - PreviousTime (%0.2f), CurrentTime (%0.2f), MoveDelta (%0.2f) "), *GetName(), NormalizedPreviousTime, NormalizedCurrentTime, Instance.DeltaTimeRecord->Delta);
 				}
 
 				Context.SetAnimationPositionRatio(NormalizedCurrentTime);
@@ -480,7 +486,7 @@ void UBlendSpace::TickAssetPlayer(FAnimTickRecord& Instance, struct FAnimNotifyQ
 				{
 					NormalizedPreviousTime = Context.GetPreviousAnimationPositionRatio();
 					NormalizedCurrentTime = Context.GetAnimationPositionRatio();
-					UE_LOG(LogAnimMarkerSync, Log, TEXT("Follower (%s) (bCanDoMarkerSync == false) - PreviousTime (%0.2f), CurrentTime (%0.2f), MoveDelta (%0.2f) "), *GetName(), NormalizedPreviousTime, NormalizedCurrentTime, MoveDelta);
+					UE_LOG(LogAnimMarkerSync, Log, TEXT("Follower (%s) (bCanDoMarkerSync == false) - PreviousTime (%0.2f), CurrentTime (%0.2f), MoveDelta (%0.2f) "), *GetName(), NormalizedPreviousTime, NormalizedCurrentTime, Instance.DeltaTimeRecord->Delta);
 				}
 			}
 
@@ -525,7 +531,7 @@ void UBlendSpace::TickAssetPlayer(FAnimTickRecord& Instance, struct FAnimNotifyQ
 
 							// Figure out delta time 
 							float DeltaTimePosition = CurrentSampleDataTime - PrevSampleDataTime;
-							const float SampleMoveDelta = MoveDelta * MultipliedSampleRateScale;
+							const float SampleMoveDelta = Instance.DeltaTimeRecord->Delta * MultipliedSampleRateScale;
 
 							// if we went against play rate, then loop around.
 							if ((SampleMoveDelta * DeltaTimePosition) < 0.f)
@@ -544,6 +550,9 @@ void UBlendSpace::TickAssetPlayer(FAnimTickRecord& Instance, struct FAnimNotifyQ
 								Context.RootMotionMovementParams.AccumulateWithBlend(Sample.Animation->ExtractRootMotion(PrevSampleDataTime, DeltaTimePosition, Instance.bLooping), SampleEntry.GetWeight());
 							}
 
+							// Capture the final adjusted delta time and previous frame time as an asset player record
+							SampleEntry.DeltaTimeRecord.Previous = PrevSampleDataTime;
+							SampleEntry.DeltaTimeRecord.Delta = DeltaTimePosition;
 
 							UE_LOG(LogAnimation, Verbose, TEXT("%d. Blending animation(%s) with %f weight at time %0.2f"), I + 1, *Sample.Animation->GetName(), SampleEntry.GetWeight(), CurrentSampleDataTime);
 						}
@@ -687,25 +696,38 @@ void UBlendSpace::ResetToRefPose(FCompactPose& OutPose) const
 	}
 }
 
+static const FAnimExtractContext DefaultBlendSpaceExtractionContext = { 0.f, true, {}, false };
+
 void UBlendSpace::GetAnimationPose(TArray<FBlendSampleData>& BlendSampleDataCache, /*out*/ FCompactPose& OutPose, /*out*/ FBlendedCurve& OutCurve) const
 {
 	UE::Anim::FStackAttributeContainer TempAttributes;
 	FAnimationPoseData AnimationPoseData = { OutPose, OutCurve, TempAttributes };
-	GetAnimationPose(BlendSampleDataCache, AnimationPoseData);
+	GetAnimationPose(BlendSampleDataCache, DefaultBlendSpaceExtractionContext, AnimationPoseData);
 }
 
 void UBlendSpace::GetAnimationPose(TArray<FBlendSampleData>& BlendSampleDataCache, /*out*/ FAnimationPoseData& OutAnimationPoseData) const
 {
-	GetAnimationPose_Internal(BlendSampleDataCache, TArrayView<FPoseLink>(), nullptr, false, OutAnimationPoseData);
+	GetAnimationPose_Internal(BlendSampleDataCache, TArrayView<FPoseLink>(), nullptr, false, DefaultBlendSpaceExtractionContext, OutAnimationPoseData);
 }
 
 void UBlendSpace::GetAnimationPose(TArray<FBlendSampleData>& BlendSampleDataCache, TArrayView<FPoseLink> InPoseLinks, /*out*/ FPoseContext& Output) const
 {
 	FAnimationPoseData AnimationPoseData(Output);
-	GetAnimationPose_Internal(BlendSampleDataCache, InPoseLinks, Output.AnimInstanceProxy, Output.ExpectsAdditivePose(), AnimationPoseData);
+	GetAnimationPose_Internal(BlendSampleDataCache, InPoseLinks, Output.AnimInstanceProxy, Output.ExpectsAdditivePose(), DefaultBlendSpaceExtractionContext, AnimationPoseData);
+}
+
+void UBlendSpace::GetAnimationPose(TArray<FBlendSampleData>& BlendSampleDataCache, const FAnimExtractContext& ExtractionContext, /*out*/ FAnimationPoseData& OutAnimationPoseData) const
+{
+	GetAnimationPose_Internal(BlendSampleDataCache, TArrayView<FPoseLink>(), nullptr, false, ExtractionContext, OutAnimationPoseData);
+}
+
+void UBlendSpace::GetAnimationPose(TArray<FBlendSampleData>& BlendSampleDataCache, TArrayView<FPoseLink> InPoseLinks, const FAnimExtractContext& ExtractionContext, /*out*/ FPoseContext& Output) const
+{
+	FAnimationPoseData AnimationPoseData(Output);
+	GetAnimationPose_Internal(BlendSampleDataCache, InPoseLinks, Output.AnimInstanceProxy, Output.ExpectsAdditivePose(), ExtractionContext, AnimationPoseData);
 }
 
-void UBlendSpace::GetAnimationPose_Internal(TArray<FBlendSampleData>& BlendSampleDataCache, TArrayView<FPoseLink> InPoseLinks, FAnimInstanceProxy* InProxy, bool bInExpectsAdditivePose, /*out*/ FAnimationPoseData& OutAnimationPoseData) const
+void UBlendSpace::GetAnimationPose_Internal(TArray<FBlendSampleData>& BlendSampleDataCache, TArrayView<FPoseLink> InPoseLinks, FAnimInstanceProxy* InProxy, bool bInExpectsAdditivePose, const FAnimExtractContext& ExtractionContext, /*out*/ FAnimationPoseData& OutAnimationPoseData) const
 {
 	SCOPE_CYCLE_COUNTER(STAT_BlendSpace_GetAnimPose);
 	FScopeCycleCounterUObject BlendSpaceScope(this);
@@ -776,7 +798,7 @@ void UBlendSpace::GetAnimationPose_Internal(TArray<FBlendSampleData>& BlendSampl
 
 					FAnimationPoseData ChildAnimationPoseData = { Pose, ChildrenCurves[I], ChildrenAttributes[I] };
 					// first one always fills up the source one
-					Sample.Animation->GetAnimationPose(ChildAnimationPoseData, FAnimExtractContext(Time, true));
+					Sample.Animation->GetAnimationPose(ChildAnimationPoseData, FAnimExtractContext(Time, ExtractionContext.bExtractRootMotion, BlendSampleDataCache[I].DeltaTimeRecord, ExtractionContext.bLooping));
 				}
 				else
 				{
diff --git a/Engine/Source/Runtime/Engine/Private/AnimationAsset.cpp b/Engine/Source/Runtime/Engine/Private/AnimationAsset.cpp
--- a/Engine/Source/Runtime/Engine/Private/AnimationAsset.cpp
+++ b/Engine/Source/Runtime/Engine/Private/AnimationAsset.cpp
@@ -205,12 +205,15 @@ FAnimTickRecord::FAnimTickRecord(UBlendSpace* InBlendSpace, const FVector& InBle
 	bLooping = bInLooping;
 }
 
-FAnimTickRecord::FAnimTickRecord(UAnimMontage* InMontage, float InCurrentPosition, float InPreviousPosition, float InMoveDelta, float InWeight, TArray<FPassedMarker>& InMarkersPassedThisTick, FMarkerTickRecord& InMarkerTickRecord)
+FAnimTickRecord::FAnimTickRecord(UAnimMontage* InMontage, float InCurrentPosition, float, float, float InWeight, TArray<FPassedMarker>& InMarkersPassedThisTick, FMarkerTickRecord& InMarkerTickRecord)
+	: FAnimTickRecord(InMontage, InCurrentPosition, InWeight, InMarkersPassedThisTick, InMarkerTickRecord)
+{
+}
+
+FAnimTickRecord::FAnimTickRecord(UAnimMontage* InMontage, float InCurrentPosition, float InWeight, TArray<FPassedMarker>& InMarkersPassedThisTick, FMarkerTickRecord& InMarkerTickRecord)
 {
 	SourceAsset = InMontage;
 	Montage.CurrentPosition = InCurrentPosition;
-	Montage.PreviousPosition = InPreviousPosition;
-	Montage.MoveDelta = InMoveDelta;
 	Montage.MarkersPassedThisTick = &InMarkersPassedThisTick;
 	MarkerTickRecord = &InMarkerTickRecord;
 	PlayRateMultiplier = 1.f; // we don't care here, this is alreayd applied in the montageinstance::Advance
diff --git a/Engine/Source/Runtime/Engine/Public/Animation/AnimRootMotionProvider.h b/Engine/Source/Runtime/Engine/Public/Animation/AnimRootMotionProvider.h
new file mode 100644
--- /dev/null
+++ b/Engine/Source/Runtime/Engine/Public/Animation/AnimRootMotionProvider.h
@@ -0,0 +1,42 @@
+// Copyright Epic Games, Inc. All Rights Reserved.
+
+#pragma once
+
+#include "CoreTypes.h"
+#include "CoreMinimal.h"
+#include "Features/IModularFeature.h"
+#include "UObject/ObjectMacros.h"
+
+class UAnimSequence;
+
+struct FDeltaTimeRecord;
+
+namespace UE { namespace Anim {
+
+struct FStackAttributeContainer;
+
+// Modular feature interface for AnimationWarping
+class ENGINE_API IAnimRootMotionProvider : public IModularFeature
+{
+public:
+	static const FName ModularFeatureName; // "AnimationWarping"
+	static const FName RootMotionDeltaAttributeName;
+
+	virtual ~IAnimRootMotionProvider() {}
+
+	static bool IsAvailable();
+	static const IAnimRootMotionProvider* Get();
+	
+	// Given the specified time range, sequence, and looping behavior, sample a root motion delta and store it in the specified attribute container
+	virtual void SampleRootMotion(const FDeltaTimeRecord& SampleRange, const UAnimSequence& Sequence, bool bLoopingSequence, FStackAttributeContainer& OutAttributes) const = 0;
+	
+	// Override the currently stored root motion delta in the specified attribute container. Requires root motion to exist prior to overriding.
+	virtual bool OverrideRootMotion(const FTransform& RootMotionDelta, FStackAttributeContainer& OutAttributes) const = 0;
+	
+	// Extract the currently stored root motion delta in the specified attribute container.
+	virtual bool ExtractRootMotion(const FStackAttributeContainer& Attributes, FTransform& OutRootMotionDelta) const = 0;
+	
+	// Query whether or not a computed root motion delta exists within the specified attribute container.
+	virtual bool HasRootMotion(const FStackAttributeContainer& Attributes) const = 0;
+};
+}}
\ No newline at end of file
diff --git a/Engine/Source/Runtime/Engine/Public/Animation/AnimSingleNodeInstanceProxy.h b/Engine/Source/Runtime/Engine/Public/Animation/AnimSingleNodeInstanceProxy.h
--- a/Engine/Source/Runtime/Engine/Public/Animation/AnimSingleNodeInstanceProxy.h
+++ b/Engine/Source/Runtime/Engine/Public/Animation/AnimSingleNodeInstanceProxy.h
@@ -60,6 +60,7 @@ public:
 		, CurrentAsset(nullptr)
 		, BlendSpacePosition(0.0f, 0.0f, 0.0f)
 		, CurrentTime(0.0f)
+		, DeltaTimeRecord()
 #if WITH_EDITORONLY_DATA
 		, PreviewPoseCurrentTime(0.0f)
 #endif
@@ -220,6 +221,8 @@ private:
 	/** Shared parameters for previewing blendspace or animsequence **/
 	float CurrentTime;
 
+	FDeltaTimeRecord DeltaTimeRecord;
+
 #if WITH_EDITORONLY_DATA
 	float PreviewPoseCurrentTime;
 #endif
