Staging compilation fix for stride warping, orientation warping behavior correction

diff --git a/Engine/Plugins/Experimental/Animation/AnimationWarping/Source/Editor/Private/AnimGraph/AnimGraphNode_StrideWarping.cpp b/Engine/Plugins/Experimental/Animation/AnimationWarping/Source/Editor/Private/AnimGraph/AnimGraphNode_StrideWarping.cpp
--- a/Engine/Plugins/Experimental/Animation/AnimationWarping/Source/Editor/Private/AnimGraph/AnimGraphNode_StrideWarping.cpp
+++ b/Engine/Plugins/Experimental/Animation/AnimationWarping/Source/Editor/Private/AnimGraph/AnimGraphNode_StrideWarping.cpp
@@ -4,6 +4,7 @@
 #include "Animation/AnimRootMotionProvider.h"
 #include "Kismet2/BlueprintEditorUtils.h"
 #include "DetailLayoutBuilder.h"
+#include "ScopedTransaction.h"
 
 #define LOCTEXT_NAMESPACE "MomentumNodes"
 
diff --git a/Engine/Plugins/Experimental/Animation/AnimationWarping/Source/Runtime/Private/BoneControllers/AnimNode_OrientationWarping.cpp b/Engine/Plugins/Experimental/Animation/AnimationWarping/Source/Runtime/Private/BoneControllers/AnimNode_OrientationWarping.cpp
--- a/Engine/Plugins/Experimental/Animation/AnimationWarping/Source/Runtime/Private/BoneControllers/AnimNode_OrientationWarping.cpp
+++ b/Engine/Plugins/Experimental/Animation/AnimationWarping/Source/Runtime/Private/BoneControllers/AnimNode_OrientationWarping.cpp
@@ -138,62 +138,64 @@ void FAnimNode_OrientationWarping::EvaluateSkeletalControl_AnyThread(FComponentS
 		RootBoneTransform.SetRotation(RootRotation * RootBoneTransform.GetRotation());
 		RootBoneTransform.NormalizeRotation();
 		Output.Pose.SetComponentSpaceTransform(RootBoneIndex, RootBoneTransform);
+	}
+		
+	const int32 NumSpineBones = SpineBoneDataArray.Num();
+	const bool bBodyOrientationAlpha = !FMath::IsNearlyZero(BodyOrientationAlpha, KINDA_SMALL_NUMBER);
+	const bool bUpdateSpineBones = (NumSpineBones > 0) && bBodyOrientationAlpha;
 
-		const float IKFootRootOrientationAlpha = 1.f - BodyOrientationAlpha;
-		const int32 NumSpineBones = SpineBoneDataArray.Num();
-		const int32 NumIKFootBones = IKFootBoneIndexArray.Num();
-		const bool bBodyOrientationAlpha = !FMath::IsNearlyZero(BodyOrientationAlpha, KINDA_SMALL_NUMBER);
-		const bool bUpdateSpineBones = (NumSpineBones > 0) && bBodyOrientationAlpha;
-		const bool bUpdateIKFootRoot = (IKFootRootBoneIndex != FCompactPoseBoneIndex(INDEX_NONE)) && !FMath::IsNearlyZero(IKFootRootOrientationAlpha, KINDA_SMALL_NUMBER);
-		const bool bUpdateIKFootBones = bUpdateIKFootRoot && (NumIKFootBones > 0);
-
-		if (bUpdateSpineBones)
+	if (bUpdateSpineBones)
+	{
+		// Spine bones counter rotate body orientation evenly across all bones.
+		for (int32 ArrayIndex = 0; ArrayIndex < NumSpineBones; ArrayIndex++)
 		{
-			// Spine bones counter rotate body orientation evenly across all bones.
-			for (int32 ArrayIndex = 0; ArrayIndex < NumSpineBones; ArrayIndex++)
-			{
-				const FOrientationWarpingSpineBoneData& BoneData = SpineBoneDataArray[ArrayIndex];
-				const FQuat SpineBoneCounterRotation = FQuat(LocomotionRotationAxis, -LocomotionRotationAngle * BodyOrientationAlpha * BoneData.Weight);
-				check(BoneData.Weight > 0.f);
-
-				FTransform SpineBoneTransform(Output.Pose.GetComponentSpaceTransform(BoneData.BoneIndex));
-				SpineBoneTransform.SetRotation((SpineBoneCounterRotation * SpineBoneTransform.GetRotation()));
-				SpineBoneTransform.NormalizeRotation();
-				Output.Pose.SetComponentSpaceTransform(BoneData.BoneIndex, SpineBoneTransform);
-			}
+			const FOrientationWarpingSpineBoneData& BoneData = SpineBoneDataArray[ArrayIndex];
+			const FQuat SpineBoneCounterRotation = FQuat(LocomotionRotationAxis, -LocomotionRotationAngle * BodyOrientationAlpha * BoneData.Weight);
+			check(BoneData.Weight > 0.f);
+
+			FTransform SpineBoneTransform(Output.Pose.GetComponentSpaceTransform(BoneData.BoneIndex));
+			SpineBoneTransform.SetRotation((SpineBoneCounterRotation * SpineBoneTransform.GetRotation()));
+			SpineBoneTransform.NormalizeRotation();
+			Output.Pose.SetComponentSpaceTransform(BoneData.BoneIndex, SpineBoneTransform);
 		}
+	}
 
-		// Rotate IK Foot Root
-		if (bUpdateIKFootRoot)
-		{
-			const FQuat BoneRotation = FQuat(LocomotionRotationAxis, LocomotionRotationAngle * IKFootRootOrientationAlpha);
+	const float IKFootRootOrientationAlpha = 1.f - BodyOrientationAlpha;
+	const bool bUpdateIKFootRoot = (IKFootRootBoneIndex != FCompactPoseBoneIndex(INDEX_NONE)) && !FMath::IsNearlyZero(IKFootRootOrientationAlpha, KINDA_SMALL_NUMBER);
 
-			FTransform IKFootRootTransform(Output.Pose.GetComponentSpaceTransform(IKFootRootBoneIndex));
-			IKFootRootTransform.SetRotation(BoneRotation * IKFootRootTransform.GetRotation());
-			IKFootRootTransform.NormalizeRotation();
-			Output.Pose.SetComponentSpaceTransform(IKFootRootBoneIndex, IKFootRootTransform);
+	// Rotate IK Foot Root
+	if (bUpdateIKFootRoot)
+	{
+		const FQuat BoneRotation = FQuat(LocomotionRotationAxis, LocomotionRotationAngle * IKFootRootOrientationAlpha);
 
-			// IK Feet 
-			// These match the root orientation, so don't rotate them. Just preserve root rotation. 
-			// We need to update their translation though, since we rotated their parent (the IK Foot Root bone).
-			if (bUpdateIKFootBones)
-			{
-				const FQuat IKFootRotation = FQuat(LocomotionRotationAxis, -LocomotionRotationAngle * IKFootRootOrientationAlpha);
+		FTransform IKFootRootTransform(Output.Pose.GetComponentSpaceTransform(IKFootRootBoneIndex));
+		IKFootRootTransform.SetRotation(BoneRotation * IKFootRootTransform.GetRotation());
+		IKFootRootTransform.NormalizeRotation();
+		Output.Pose.SetComponentSpaceTransform(IKFootRootBoneIndex, IKFootRootTransform);
 
-				for (int32 ArrayIndex = 0; ArrayIndex < NumIKFootBones; ArrayIndex++)
-				{
-					const FCompactPoseBoneIndex& IKFootBoneIndex = IKFootBoneIndexArray[ArrayIndex];
+		// IK Feet 
+		// These match the root orientation, so don't rotate them. Just preserve root rotation. 
+		// We need to update their translation though, since we rotated their parent (the IK Foot Root bone).
+		const int32 NumIKFootBones = IKFootBoneIndexArray.Num();
+		const bool bUpdateIKFootBones = bUpdateIKFootRoot && (NumIKFootBones > 0);
 
-					FTransform IKFootBoneTransform(Output.Pose.GetComponentSpaceTransform(IKFootBoneIndex));
-					IKFootBoneTransform.SetRotation(IKFootRotation * IKFootBoneTransform.GetRotation());
-					IKFootBoneTransform.NormalizeRotation();
-					Output.Pose.SetComponentSpaceTransform(IKFootBoneIndex, IKFootBoneTransform);
-				}
+		if (bUpdateIKFootBones)
+		{
+			const FQuat IKFootRotation = FQuat(LocomotionRotationAxis, -LocomotionRotationAngle * IKFootRootOrientationAlpha);
+
+			for (int32 ArrayIndex = 0; ArrayIndex < NumIKFootBones; ArrayIndex++)
+			{
+				const FCompactPoseBoneIndex& IKFootBoneIndex = IKFootBoneIndexArray[ArrayIndex];
+
+				FTransform IKFootBoneTransform(Output.Pose.GetComponentSpaceTransform(IKFootBoneIndex));
+				IKFootBoneTransform.SetRotation(IKFootRotation * IKFootBoneTransform.GetRotation());
+				IKFootBoneTransform.NormalizeRotation();
+				Output.Pose.SetComponentSpaceTransform(IKFootBoneIndex, IKFootBoneTransform);
 			}
 		}
-
-		OutBoneTransforms.Sort(FCompareBoneTransformIndex());
 	}
+
+	OutBoneTransforms.Sort(FCompareBoneTransformIndex());
 }
 
 bool FAnimNode_OrientationWarping::IsValidToEvaluate(const USkeleton* Skeleton, const FBoneContainer& RequiredBones)
