Stride Warping: Fix rapid change in the stride scale when using graph mode and root motion speed or locomotion speed is nearly zero.
Orientation Warping: Potential solution to prevent over rotation of the lower body when using graph mode and root motion deviates too much the movement direction + added draw debugs to visualize root motion direction and physics direction + allow the alpha of the node to affect the final rotation

diff --git a/Engine/Plugins/Experimental/Animation/AnimationWarping/Source/Runtime/Private/BoneControllers/AnimNode_OrientationWarping.cpp b/Engine/Plugins/Experimental/Animation/AnimationWarping/Source/Runtime/Private/BoneControllers/AnimNode_OrientationWarping.cpp
--- a/Engine/Plugins/Experimental/Animation/AnimationWarping/Source/Runtime/Private/BoneControllers/AnimNode_OrientationWarping.cpp
+++ b/Engine/Plugins/Experimental/Animation/AnimationWarping/Source/Runtime/Private/BoneControllers/AnimNode_OrientationWarping.cpp
@@ -7,6 +7,11 @@
 
 DECLARE_CYCLE_STAT(TEXT("OrientationWarping Eval"), STAT_OrientationWarping_Eval, STATGROUP_Anim);
 
+#if ENABLE_ANIM_DEBUG
+TAutoConsoleVariable<int32> CVarAnimNodeOrientationWarpingDebug(TEXT("a.AnimNode.OrientationWarping.Debug"), 0, TEXT("Turn on debug for AnimNode_OrientationWarping"));
+TAutoConsoleVariable<int32> CVarAnimNodeOrientationWarpingEnable(TEXT("a.AnimNode.OrientationWarping.Enable"), 1, TEXT("Toggle Orientation Warping"));
+#endif
+
 FAnimNode_OrientationWarping::FAnimNode_OrientationWarping()
 	: Mode(EWarpingEvaluationMode::Manual)
 	, LocomotionAngle(0.f)
@@ -78,11 +83,6 @@ void FAnimNode_OrientationWarping::EvaluateSkeletalControl_AnyThread(FComponentS
 	// Graph driven orientation warping will modify the incoming root motion to orient towards the intended locomotion angle
 	if (bGraphDrivenWarping)
 	{
-		if (IsInvalidWarpingAngle(LocomotionRotationAngle, KINDA_SMALL_NUMBER))
-		{
-			return;
-		}
-
 		FTransform RootMotionTransformDelta;
 		const bool bRootMotionDeltaPresent = RootMotionProvider->ExtractRootMotion(Output.CustomAttributes, RootMotionTransformDelta);
 
@@ -103,12 +103,17 @@ void FAnimNode_OrientationWarping::EvaluateSkeletalControl_AnyThread(FComponentS
 			const FVector RootMotionDeltaDir = RootMotionTransformDelta.GetTranslation();
 
 			// Capture the delta rotation from the axis of motion we care about
-			const FQuat WarpedRotation = FQuat::FindBetween(RootMotionDeltaDir, LocomotionForwardDir);
+			FQuat WarpedRotation = FQuat::FindBetween(RootMotionDeltaDir, LocomotionForwardDir);
 			LocomotionRotationAngle = WarpedRotation.GetTwistAngle(LocomotionRotationAxis);
 
-			if (IsInvalidWarpingAngle(LocomotionRotationAngle, KINDA_SMALL_NUMBER))
+			// Motion Matching may return an animation that deviates a lot from the movement direction (e.g movement direction going bwd and motion matching could return the fwd animation for a few frames)
+			// When that happens, since we use the delta between root motion and movement direction, we would be over-rotating the lower body and breaking the pose during those frames
+			// So, when that happens we use the inverse of the movement direction to calculate our target rotation. 
+			// This feels a bit 'hacky' but its the only option I've found so far to mitigate the problem
+			if (LocomotionInversionThresholdAngle > 0.f && FMath::Abs(FMath::RadiansToDegrees(LocomotionRotationAngle)) > LocomotionInversionThresholdAngle)
 			{
-				return;
+				WarpedRotation = FQuat::FindBetween(RootMotionDeltaDir, -LocomotionForwardDir);
+				LocomotionRotationAngle = WarpedRotation.GetTwistAngle(LocomotionRotationAxis);
 			}
 
 			// Rotate the root motion delta fully by the warped angle
@@ -119,6 +124,26 @@ void FAnimNode_OrientationWarping::EvaluateSkeletalControl_AnyThread(FComponentS
 			// Forward the side effects of orientation warping on the root motion contribution for this sub-graph
 			const bool bRootMotionOverridden = RootMotionProvider->OverrideRootMotion(RootMotionTransformDelta, Output.CustomAttributes);
 			ensure(bRootMotionOverridden);
+
+#if ENABLE_ANIM_DEBUG
+			const bool bShowDebug = (CVarAnimNodeOrientationWarpingDebug.GetValueOnAnyThread() == 1);
+			if (bShowDebug)
+			{
+				const FVector LocomotionDirWS = Output.AnimInstanceProxy->GetActorTransform().TransformVectorNoScale(LocomotionRotation.GetForwardVector()).GetSafeNormal();
+				const FVector RootMotionDirWS = Output.AnimInstanceProxy->GetActorTransform().TransformVectorNoScale(SkeletalMeshRelativeRotation.RotateVector(RootMotionDeltaDir)).GetSafeNormal();
+
+				const float DrawDebugArrowScale = 100.f;
+				Output.AnimInstanceProxy->AnimDrawDebugLine(
+					Output.AnimInstanceProxy->GetComponentTransform().GetLocation(),
+					Output.AnimInstanceProxy->GetComponentTransform().GetLocation() + LocomotionDirWS * DrawDebugArrowScale,
+					FColor::Red, false, 0.f, 1.f);
+
+				Output.AnimInstanceProxy->AnimDrawDebugLine(
+					Output.AnimInstanceProxy->GetComponentTransform().GetLocation(),
+					Output.AnimInstanceProxy->GetComponentTransform().GetLocation() + RootMotionDirWS * DrawDebugArrowScale,
+					FColor::Blue, false, 0.f, 1.f);
+			}
+#endif
 		}
 	} 
 	else if (IsInvalidWarpingAngle(LocomotionRotationAngle, KINDA_SMALL_NUMBER))
@@ -126,6 +151,15 @@ void FAnimNode_OrientationWarping::EvaluateSkeletalControl_AnyThread(FComponentS
 		return;
 	}
 
+	if (InterpSpeed > 0.f)
+	{
+		LocomotionRotationAngle = FMath::FInterpTo(LastRotationAngle, LocomotionRotationAngle, Output.AnimInstanceProxy->GetDeltaSeconds(), InterpSpeed);
+		LastRotationAngle = LocomotionRotationAngle;
+	}
+
+	// Allow the alpha value of the node to affect the final rotation
+	LocomotionRotationAngle *= ActualAlpha;
+
 	const float BodyOrientationAlpha = FMath::Clamp(Settings.BodyOrientationAlpha, 0.f, 1.f);
 
 	// Rotate Root Bone first, as that cheaply rotates the whole pose with one transformation.
@@ -200,6 +234,13 @@ void FAnimNode_OrientationWarping::EvaluateSkeletalControl_AnyThread(FComponentS
 
 bool FAnimNode_OrientationWarping::IsValidToEvaluate(const USkeleton* Skeleton, const FBoneContainer& RequiredBones)
 {
+#if ENABLE_ANIM_DEBUG
+	if (CVarAnimNodeOrientationWarpingEnable.GetValueOnAnyThread() == 0)
+	{
+		return false;
+	}
+#endif
+
 	bool bIKFootRootIsValid = IKFootRootBoneIndex != INDEX_NONE;
 	bool bIKFeetAreValid = IKFootBoneIndexArray.Num() > 0;
 	for (const auto& IKFootBoneIndex : IKFootBoneIndexArray)
diff --git a/Engine/Plugins/Experimental/Animation/AnimationWarping/Source/Runtime/Private/BoneControllers/AnimNode_StrideWarping.cpp b/Engine/Plugins/Experimental/Animation/AnimationWarping/Source/Runtime/Private/BoneControllers/AnimNode_StrideWarping.cpp
--- a/Engine/Plugins/Experimental/Animation/AnimationWarping/Source/Runtime/Private/BoneControllers/AnimNode_StrideWarping.cpp
+++ b/Engine/Plugins/Experimental/Animation/AnimationWarping/Source/Runtime/Private/BoneControllers/AnimNode_StrideWarping.cpp
@@ -9,10 +9,9 @@
 
 #if ENABLE_ANIM_DEBUG
 TAutoConsoleVariable<int32> CVarAnimNodeStrideWarpingDebug(TEXT("a.AnimNode.StrideWarping.Debug"), 0, TEXT("Turn on debug for AnimNode_StrideWarping"));
+TAutoConsoleVariable<int32> CVarAnimNodeStrideWarpingEnable(TEXT("a.AnimNode.StrideWarping.Enable"), 1, TEXT("Toggle Stride Warping"));
 #endif
 
-TAutoConsoleVariable<int32> CVarStrideWarpingEnable(TEXT("a.AnimNode.StrideWarping.Enable"), 1, TEXT("Toggle Stride Warping"));
-
 DECLARE_CYCLE_STAT(TEXT("StrideWarping Eval"), STAT_StrideWarping_Eval, STATGROUP_Anim);
 
 FAnimNode_StrideWarping::FAnimNode_StrideWarping()
@@ -151,14 +150,17 @@ void FAnimNode_StrideWarping::EvaluateSkeletalControl_AnyThread(FComponentSpaceP
 		}
 
 		const float RootMotionSpeed = RootMotionTransformDelta.GetTranslation().Size() / CachedDeltaTime;
-		if (FMath::IsNearlyZero(RootMotionSpeed, KINDA_SMALL_NUMBER))
+
+		if (FMath::IsNearlyZero(RootMotionSpeed, MinSpeedTolerance) || FMath::IsNearlyZero(LocomotionSpeed, MinSpeedTolerance))
 		{
-			return;
+			ComputedStrideScaling = 1.f;
+		}
+		else
+		{
+			// Graph driven stride scale factor will be determined by the ratio of the
+			// locomotion (capsule/physics) speed against the animation root motion speed
+			ComputedStrideScaling = LocomotionSpeed / RootMotionSpeed;
 		}
-
-		// Graph driven stride scale factor will be determined by the ratio of the
-		// locomotion (capsule/physics) speed against the animation root motion speed
-		ComputedStrideScaling = LocomotionSpeed / RootMotionSpeed;
 	}
 
 	// Allow the opportunity for stride scale clamping and biasing regardless of evaluation mode
@@ -347,8 +349,14 @@ void FAnimNode_StrideWarping::EvaluateSkeletalControl_AnyThread(FComponentSpaceP
 
 bool FAnimNode_StrideWarping::IsValidToEvaluate(const USkeleton* Skeleton, const FBoneContainer& RequiredBones)
 {
-	const bool bIsEnabled = (CVarStrideWarpingEnable.GetValueOnAnyThread() == 1);
-	return bIsEnabled && (FeetData.Num() > 0) 
+#if ENABLE_ANIM_DEBUG
+	if(CVarAnimNodeStrideWarpingEnable.GetValueOnAnyThread() == 0)
+	{
+		return false;
+	}
+#endif
+
+	return (FeetData.Num() > 0) 
 		&& (PelvisBone.GetCompactPoseIndex(RequiredBones) != INDEX_NONE) 
 		&& (IKFootRootBone.GetCompactPoseIndex(RequiredBones) != INDEX_NONE) 
 		&& (!FMath::IsNearlyEqual(StrideScalingScaleBiasClamp.ApplyTo(ManualStrideScaling, 0.f), 1.f, 0.001f) || PelvisAdjustmentInterp.IsInMotion() || Mode == EWarpingEvaluationMode::Graph);
diff --git a/Engine/Plugins/Experimental/Animation/AnimationWarping/Source/Runtime/Public/BoneControllers/AnimNode_OrientationWarping.h b/Engine/Plugins/Experimental/Animation/AnimationWarping/Source/Runtime/Public/BoneControllers/AnimNode_OrientationWarping.h
--- a/Engine/Plugins/Experimental/Animation/AnimationWarping/Source/Runtime/Public/BoneControllers/AnimNode_OrientationWarping.h
+++ b/Engine/Plugins/Experimental/Animation/AnimationWarping/Source/Runtime/Public/BoneControllers/AnimNode_OrientationWarping.h
@@ -91,6 +91,15 @@ struct ANIMATIONWARPINGRUNTIME_API FAnimNode_OrientationWarping : public FAnimNo
 
 	FCompactPoseBoneIndex IKFootRootBoneIndex;
 
+	UPROPERTY(EditAnywhere, BlueprintReadWrite, Category = Settings, meta = (PinHiddenByDefault, EditCondition = "Mode == EWarpingEvaluationMode::Graph"))
+	float LocomotionInversionThresholdAngle = 90.f;
+
+	UPROPERTY(EditAnywhere, BlueprintReadWrite, Category = Settings, meta = (PinHiddenByDefault))
+	float InterpSpeed = 0.f;
+
+	UPROPERTY(Transient)
+	float LastRotationAngle = 0.f;
+
 public:
 	FAnimNode_OrientationWarping();
 
diff --git a/Engine/Plugins/Experimental/Animation/AnimationWarping/Source/Runtime/Public/BoneControllers/AnimNode_StrideWarping.h b/Engine/Plugins/Experimental/Animation/AnimationWarping/Source/Runtime/Public/BoneControllers/AnimNode_StrideWarping.h
--- a/Engine/Plugins/Experimental/Animation/AnimationWarping/Source/Runtime/Public/BoneControllers/AnimNode_StrideWarping.h
+++ b/Engine/Plugins/Experimental/Animation/AnimationWarping/Source/Runtime/Public/BoneControllers/AnimNode_StrideWarping.h
@@ -133,6 +133,9 @@ struct ANIMATIONWARPINGRUNTIME_API FAnimNode_StrideWarping : public FAnimNode_Sk
 	UPROPERTY(Transient)
 	float CachedDeltaTime;
 
+	UPROPERTY(EditAnywhere, BlueprintReadWrite, Category = Settings, meta = (PinHiddenByDefault, EditCondition = "Mode == EWarpingEvaluationMode::Graph"))
+	float MinSpeedTolerance = 10.f;
+
 	FVector GetAxisModeValue(const EStrideWarpingAxisMode& AxisMode, const FTransform& IKFootRootCSTransform, const FVector& UserSuppliedVector) const;
 	FCompactPoseBoneIndex FindHipBoneIndex(const FCompactPoseBoneIndex& InFootBoneIndex, const int32& NumBonesInLimb, const FBoneContainer& RequiredBones) const;
 
