Stride and Orientation Warping interface cleanup for UE5 Release

diff --git a/Engine/Plugins/Experimental/Animation/AnimationWarping/Source/Editor/AnimationWarpingEditor.Build.cs b/Engine/Plugins/Experimental/Animation/AnimationWarping/Source/Editor/AnimationWarpingEditor.Build.cs
--- a/Engine/Plugins/Experimental/Animation/AnimationWarping/Source/Editor/AnimationWarpingEditor.Build.cs
+++ b/Engine/Plugins/Experimental/Animation/AnimationWarping/Source/Editor/AnimationWarpingEditor.Build.cs
@@ -9,12 +9,13 @@ namespace UnrealBuildTool.Rules
 			PublicDependencyModuleNames.AddRange(
 			new string[]
 			{
+				"AnimationCore",
+				"AnimationWarpingRuntime",
                 "AnimGraph",
-                "AnimationCore",
+				"AnimGraphRuntime",
                 "Core",
                 "CoreUObject",
                 "Engine",
-				"AnimationWarpingRuntime",
             });
 
 			PrivateDependencyModuleNames.AddRange(
@@ -28,12 +29,10 @@ namespace UnrealBuildTool.Rules
                 PrivateDependencyModuleNames.AddRange(
                     new string[]
                     {
+						"BlueprintGraph",
 						"EditorFramework",
+						"Kismet",
                         "UnrealEd",
-                        "Kismet",
-                        "AnimGraph",
-						"AnimGraphRuntime",
-                        "BlueprintGraph",
                     }
                 );
             }
diff --git a/Engine/Plugins/Experimental/Animation/AnimationWarping/Source/Editor/Private/AnimGraph/AnimGraphNode_OrientationWarping.cpp b/Engine/Plugins/Experimental/Animation/AnimationWarping/Source/Editor/Private/AnimGraph/AnimGraphNode_OrientationWarping.cpp
--- a/Engine/Plugins/Experimental/Animation/AnimationWarping/Source/Editor/Private/AnimGraph/AnimGraphNode_OrientationWarping.cpp
+++ b/Engine/Plugins/Experimental/Animation/AnimationWarping/Source/Editor/Private/AnimGraph/AnimGraphNode_OrientationWarping.cpp
@@ -3,15 +3,13 @@
 #include "AnimGraph/AnimGraphNode_OrientationWarping.h"
 #include "Animation/AnimRootMotionProvider.h"
 #include "Kismet2/BlueprintEditorUtils.h"
+#include "DetailCategoryBuilder.h"
 #include "DetailLayoutBuilder.h"
+#include "ScopedTransaction.h"
 
-/////////////////////////////////////////////////////
-// UAnimGraphNode_OrientationWarping
+#define LOCTEXT_NAMESPACE "AnimationWarping"
 
-#define LOCTEXT_NAMESPACE "MomentumNodes"
-
-UAnimGraphNode_OrientationWarping::UAnimGraphNode_OrientationWarping(const FObjectInitializer& ObjectInitializer)
-	: Super(ObjectInitializer)
+UAnimGraphNode_OrientationWarping::UAnimGraphNode_OrientationWarping(const FObjectInitializer& ObjectInitializer) : Super(ObjectInitializer)
 {
 }
 
@@ -22,7 +20,7 @@ FText UAnimGraphNode_OrientationWarping::GetControllerDescription() const
 
 FText UAnimGraphNode_OrientationWarping::GetTooltipText() const
 {
-	return LOCTEXT("OrientationWarpingTooltip", "Orients Root Bone to match locomotion direction, and counter rotates spine.");
+	return LOCTEXT("OrientationWarpingTooltip", "Warps the root bone, and counter rotates the body, to match locomotion direction.");
 }
 
 FText UAnimGraphNode_OrientationWarping::GetNodeTitle(ENodeTitleType::Type TitleType) const
@@ -30,6 +28,86 @@ FText UAnimGraphNode_OrientationWarping::GetNodeTitle(ENodeTitleType::Type Title
 	return GetControllerDescription();
 }
 
+FLinearColor UAnimGraphNode_OrientationWarping::GetNodeTitleColor() const
+{
+	return FLinearColor(FColor(153.f, 0.f, 0.f));
+}
+
+void UAnimGraphNode_OrientationWarping::CustomizePinData(UEdGraphPin* Pin, FName SourcePropertyName, int32 ArrayIndex) const
+{
+	Super::CustomizePinData(Pin, SourcePropertyName, ArrayIndex);
+
+	if (Pin->PinName == GET_MEMBER_NAME_STRING_CHECKED(FAnimNode_OrientationWarping, OrientationAngle))
+	{
+		Pin->bHidden = (Node.Mode == EWarpingEvaluationMode::Graph);
+	}
+
+	if (Pin->PinName == GET_MEMBER_NAME_STRING_CHECKED(FAnimNode_OrientationWarping, LocomotionAngle))
+	{
+		Pin->bHidden = (Node.Mode == EWarpingEvaluationMode::Manual);
+	}
+
+	if (Pin->PinName == GET_MEMBER_NAME_STRING_CHECKED(FAnimNode_OrientationWarping, LocomotionAngleDeltaThreshold))
+	{
+		Pin->bHidden = (Node.Mode == EWarpingEvaluationMode::Manual);
+	}
+}
+
+void UAnimGraphNode_OrientationWarping::CustomizeDetails(IDetailLayoutBuilder& DetailBuilder)
+{
+	DECLARE_SCOPE_HIERARCHICAL_COUNTER_FUNC()
+	Super::CustomizeDetails(DetailBuilder);
+
+	DetailBuilder.SortCategories([](const TMap<FName, IDetailCategoryBuilder*>& CategoryMap)
+	{
+		for (const TPair<FName, IDetailCategoryBuilder*>& Pair : CategoryMap)
+		{
+			int32 SortOrder = Pair.Value->GetSortOrder();
+			const FName CategoryName = Pair.Key;
+
+			if (CategoryName == "Evaluation")
+			{
+				SortOrder += 1;
+			}
+			else if (CategoryName == "Settings")
+			{
+				SortOrder += 2;
+			}
+			else if (CategoryName == "Debug")
+			{
+				SortOrder += 3;
+			}
+			else
+			{
+				const int32 ValueSortOrder = Pair.Value->GetSortOrder();
+				if (ValueSortOrder >= SortOrder && ValueSortOrder < SortOrder + 10)
+				{
+					SortOrder += 10;
+				}
+				else
+				{
+					continue;
+				}
+			}
+
+			Pair.Value->SetSortOrder(SortOrder);
+		}
+	});
+
+	TSharedRef<IPropertyHandle> NodeHandle = DetailBuilder.GetProperty(FName(TEXT("Node")), GetClass());
+
+	if (Node.Mode == EWarpingEvaluationMode::Graph)
+	{
+		DetailBuilder.HideProperty(NodeHandle->GetChildHandle(GET_MEMBER_NAME_CHECKED(FAnimNode_OrientationWarping, OrientationAngle)));
+	}
+
+	if (Node.Mode == EWarpingEvaluationMode::Manual)
+	{
+		DetailBuilder.HideProperty(NodeHandle->GetChildHandle(GET_MEMBER_NAME_CHECKED(FAnimNode_OrientationWarping, LocomotionAngle)));
+		DetailBuilder.HideProperty(NodeHandle->GetChildHandle(GET_MEMBER_NAME_CHECKED(FAnimNode_OrientationWarping, LocomotionAngleDeltaThreshold)));
+	}
+}
+
 void UAnimGraphNode_OrientationWarping::PostEditChangeProperty(struct FPropertyChangedEvent& PropertyChangedEvent)
 {
 	DECLARE_SCOPE_HIERARCHICAL_COUNTER_FUNC()
@@ -40,18 +118,39 @@ void UAnimGraphNode_OrientationWarping::PostEditChangeProperty(struct FPropertyC
 
 	if (ChangedProperty)
 	{
-		if ((ChangedProperty->GetFName() == GET_MEMBER_NAME_STRING_CHECKED(FInputScaleBiasClamp, bMapRange))
-			|| (ChangedProperty->GetFName() == GET_MEMBER_NAME_STRING_CHECKED(FInputRange, Min))
-			|| (ChangedProperty->GetFName() == GET_MEMBER_NAME_STRING_CHECKED(FInputRange, Max))
-			|| (ChangedProperty->GetFName() == GET_MEMBER_NAME_STRING_CHECKED(FInputScaleBiasClamp, Scale))
-			|| (ChangedProperty->GetFName() == GET_MEMBER_NAME_STRING_CHECKED(FInputScaleBiasClamp, Bias))
-			|| (ChangedProperty->GetFName() == GET_MEMBER_NAME_STRING_CHECKED(FInputScaleBiasClamp, bClampResult))
-			|| (ChangedProperty->GetFName() == GET_MEMBER_NAME_STRING_CHECKED(FInputScaleBiasClamp, ClampMin))
-			|| (ChangedProperty->GetFName() == GET_MEMBER_NAME_STRING_CHECKED(FInputScaleBiasClamp, ClampMax))
-			|| (ChangedProperty->GetFName() == GET_MEMBER_NAME_STRING_CHECKED(FInputScaleBiasClamp, bInterpResult))
-			|| (ChangedProperty->GetFName() == GET_MEMBER_NAME_STRING_CHECKED(FInputScaleBiasClamp, InterpSpeedIncreasing))
-			|| (ChangedProperty->GetFName() == GET_MEMBER_NAME_STRING_CHECKED(FInputScaleBiasClamp, InterpSpeedDecreasing)))
+		// Evaluation mode
+		if (ChangedProperty->GetFName() == GET_MEMBER_NAME_STRING_CHECKED(FAnimNode_OrientationWarping, Mode))
 		{
+			FScopedTransaction Transaction(LOCTEXT("ChangeEvaluationMode", "Change Evaluation Mode"));
+			Modify();
+
+			// Break links to pins going away
+			for (int32 PinIndex = 0; PinIndex < Pins.Num(); ++PinIndex)
+			{
+				UEdGraphPin* Pin = Pins[PinIndex];
+				if (Pin->PinName == GET_MEMBER_NAME_STRING_CHECKED(FAnimNode_OrientationWarping, OrientationAngle))
+				{
+					if (Node.Mode == EWarpingEvaluationMode::Graph)
+					{
+						Pin->BreakAllPinLinks();
+					}
+				}
+				else if (Pin->PinName == GET_MEMBER_NAME_STRING_CHECKED(FAnimNode_OrientationWarping, LocomotionAngle))
+				{
+					if (Node.Mode == EWarpingEvaluationMode::Manual)
+					{
+						Pin->BreakAllPinLinks();
+					}
+				}
+				else if (Pin->PinName == GET_MEMBER_NAME_STRING_CHECKED(FAnimNode_OrientationWarping, LocomotionAngleDeltaThreshold))
+				{
+					if (Node.Mode == EWarpingEvaluationMode::Manual)
+					{
+						Pin->BreakAllPinLinks();
+					}
+				}
+			}
+
 			bRequiresNodeReconstruct = true;
 		}
 	}
diff --git a/Engine/Plugins/Experimental/Animation/AnimationWarping/Source/Editor/Private/AnimGraph/AnimGraphNode_SlopeWarping.cpp b/Engine/Plugins/Experimental/Animation/AnimationWarping/Source/Editor/Private/AnimGraph/AnimGraphNode_SlopeWarping.cpp
--- a/Engine/Plugins/Experimental/Animation/AnimationWarping/Source/Editor/Private/AnimGraph/AnimGraphNode_SlopeWarping.cpp
+++ b/Engine/Plugins/Experimental/Animation/AnimationWarping/Source/Editor/Private/AnimGraph/AnimGraphNode_SlopeWarping.cpp
@@ -1,10 +1,8 @@
 // Copyright Epic Games, Inc. All Rights Reserved.
 
-
 #include "AnimGraph/AnimGraphNode_SlopeWarping.h"
-#include "Animation/AnimRootMotionProvider.h"
 
-#define LOCTEXT_NAMESPACE "MomentumNodes"
+#define LOCTEXT_NAMESPACE "AnimationWarping"
 
 UAnimGraphNode_SlopeWarping::UAnimGraphNode_SlopeWarping(const FObjectInitializer& ObjectInitializer) : Super(ObjectInitializer)
 {
@@ -17,7 +15,7 @@ FText UAnimGraphNode_SlopeWarping::GetControllerDescription() const
 
 FText UAnimGraphNode_SlopeWarping::GetTooltipText() const
 {
-	return LOCTEXT("SlopeWarpingTooltip", "Adjust Feet IK to Match Floor Normal");
+	return LOCTEXT("SlopeWarpingTooltip", "Warps the feet to match the floor normal.");
 }
 
 FText UAnimGraphNode_SlopeWarping::GetNodeTitle(ENodeTitleType::Type TitleType) const
@@ -25,20 +23,9 @@ FText UAnimGraphNode_SlopeWarping::GetNodeTitle(ENodeTitleType::Type TitleType)
 	return GetControllerDescription();
 }
 
-void UAnimGraphNode_SlopeWarping::GetInputLinkAttributes(FNodeAttributeArray& OutAttributes) const
-{
-	if (Node.Mode == EWarpingEvaluationMode::Graph)
-	{
-		OutAttributes.Add(UE::Anim::IAnimRootMotionProvider::AttributeName);
-	}
-}
-
-void UAnimGraphNode_SlopeWarping::GetOutputLinkAttributes(FNodeAttributeArray& OutAttributes) const
+FLinearColor UAnimGraphNode_SlopeWarping::GetNodeTitleColor() const
 {
-	if (Node.Mode == EWarpingEvaluationMode::Graph)
-	{
-		OutAttributes.Add(UE::Anim::IAnimRootMotionProvider::AttributeName);
-	}
+	return FLinearColor(FColor(153.f, 0.f, 0.f));
 }
 
 #undef LOCTEXT_NAMESPACE
diff --git a/Engine/Plugins/Experimental/Animation/AnimationWarping/Source/Editor/Private/AnimGraph/AnimGraphNode_StrideWarping.cpp b/Engine/Plugins/Experimental/Animation/AnimationWarping/Source/Editor/Private/AnimGraph/AnimGraphNode_StrideWarping.cpp
--- a/Engine/Plugins/Experimental/Animation/AnimationWarping/Source/Editor/Private/AnimGraph/AnimGraphNode_StrideWarping.cpp
+++ b/Engine/Plugins/Experimental/Animation/AnimationWarping/Source/Editor/Private/AnimGraph/AnimGraphNode_StrideWarping.cpp
@@ -3,10 +3,11 @@
 #include "AnimGraph/AnimGraphNode_StrideWarping.h"
 #include "Animation/AnimRootMotionProvider.h"
 #include "Kismet2/BlueprintEditorUtils.h"
+#include "DetailCategoryBuilder.h"
 #include "DetailLayoutBuilder.h"
 #include "ScopedTransaction.h"
 
-#define LOCTEXT_NAMESPACE "MomentumNodes"
+#define LOCTEXT_NAMESPACE "AnimationWarping"
 
 UAnimGraphNode_StrideWarping::UAnimGraphNode_StrideWarping(const FObjectInitializer& ObjectInitializer) : Super(ObjectInitializer)
 {
@@ -19,7 +20,7 @@ FText UAnimGraphNode_StrideWarping::GetControllerDescription() const
 
 FText UAnimGraphNode_StrideWarping::GetTooltipText() const
 {
-	return LOCTEXT("StrideWarpingTooltip", "Scale Feet IK to match movement speed.");
+	return LOCTEXT("StrideWarpingTooltip", "Warps the pelvis and feet to match the locomotion speed.");
 }
 
 FText UAnimGraphNode_StrideWarping::GetNodeTitle(ENodeTitleType::Type TitleType) const
@@ -27,23 +28,23 @@ FText UAnimGraphNode_StrideWarping::GetNodeTitle(ENodeTitleType::Type TitleType)
 	return GetControllerDescription();
 }
 
+FLinearColor UAnimGraphNode_StrideWarping::GetNodeTitleColor() const
+{
+	return FLinearColor(FColor(153.f, 0.f, 0.f));
+}
+
 void UAnimGraphNode_StrideWarping::CustomizePinData(UEdGraphPin* Pin, FName SourcePropertyName, int32 ArrayIndex) const
 {
 	Super::CustomizePinData(Pin, SourcePropertyName, ArrayIndex);
 
-	if (Pin->PinName == GET_MEMBER_NAME_STRING_CHECKED(FAnimNode_StrideWarping, ManualStrideWarpingDir))
+	if (Pin->PinName == GET_MEMBER_NAME_STRING_CHECKED(FAnimNode_StrideWarping, StrideDirection))
 	{
 		Pin->bHidden = (Node.Mode == EWarpingEvaluationMode::Graph);
 	}
 
-	if (Pin->PinName == GET_MEMBER_NAME_STRING_CHECKED(FAnimNode_StrideWarping, ManualStrideScaling))
+	if (Pin->PinName == GET_MEMBER_NAME_STRING_CHECKED(FAnimNode_StrideWarping, StrideScale))
 	{
 		Pin->bHidden = (Node.Mode == EWarpingEvaluationMode::Graph);
-
-		if (!Pin->bHidden)
-		{
-			Pin->PinFriendlyName = Node.StrideScalingScaleBiasClamp.GetFriendlyName(Pin->PinFriendlyName);
-		}
 	}
 
 	if (Pin->PinName == GET_MEMBER_NAME_STRING_CHECKED(FAnimNode_StrideWarping, LocomotionSpeed))
@@ -51,7 +52,7 @@ void UAnimGraphNode_StrideWarping::CustomizePinData(UEdGraphPin* Pin, FName Sour
 		Pin->bHidden = (Node.Mode == EWarpingEvaluationMode::Manual);
 	}
 
-	if (Pin->PinName == GET_MEMBER_NAME_STRING_CHECKED(FAnimNode_StrideWarping, MinSpeedTolerance))
+	if (Pin->PinName == GET_MEMBER_NAME_STRING_CHECKED(FAnimNode_StrideWarping, MinLocomotionSpeedThreshold))
 	{
 		Pin->bHidden = (Node.Mode == EWarpingEvaluationMode::Manual);
 	}
@@ -62,18 +63,58 @@ void UAnimGraphNode_StrideWarping::CustomizeDetails(IDetailLayoutBuilder& Detail
 	DECLARE_SCOPE_HIERARCHICAL_COUNTER_FUNC()
 	Super::CustomizeDetails(DetailBuilder);
 
+	DetailBuilder.SortCategories([](const TMap<FName, IDetailCategoryBuilder*>& CategoryMap)
+	{
+		for (const TPair<FName, IDetailCategoryBuilder*>& Pair : CategoryMap)
+		{
+			int32 SortOrder = Pair.Value->GetSortOrder();
+			const FName CategoryName = Pair.Key;
+
+			if (CategoryName == "Evaluation")
+			{
+				SortOrder += 1;
+			}
+			else if (CategoryName == "Settings")
+			{
+				SortOrder += 2;
+			}
+			else if (CategoryName == "Advanced")
+			{
+				SortOrder += 3;
+			}
+			else if (CategoryName == "Debug")
+			{
+				SortOrder += 4;
+			}
+			else
+			{
+				const int32 ValueSortOrder = Pair.Value->GetSortOrder();
+				if (ValueSortOrder >= SortOrder && ValueSortOrder < SortOrder + 10)
+				{
+					SortOrder += 10;
+				}
+				else
+				{
+					continue;
+				}
+			}
+
+			Pair.Value->SetSortOrder(SortOrder);
+		}
+	});
+
 	TSharedRef<IPropertyHandle> NodeHandle = DetailBuilder.GetProperty(FName(TEXT("Node")), GetClass());
 
 	if (Node.Mode == EWarpingEvaluationMode::Graph)
 	{
-		DetailBuilder.HideProperty(NodeHandle->GetChildHandle(GET_MEMBER_NAME_CHECKED(FAnimNode_StrideWarping, ManualStrideWarpingDir)));
-		DetailBuilder.HideProperty(NodeHandle->GetChildHandle(GET_MEMBER_NAME_CHECKED(FAnimNode_StrideWarping, ManualStrideScaling)));
+		DetailBuilder.HideProperty(NodeHandle->GetChildHandle(GET_MEMBER_NAME_CHECKED(FAnimNode_StrideWarping, StrideScale)));
+		DetailBuilder.HideProperty(NodeHandle->GetChildHandle(GET_MEMBER_NAME_CHECKED(FAnimNode_StrideWarping, StrideDirection)));
 	}
 
 	if (Node.Mode == EWarpingEvaluationMode::Manual)
 	{
 		DetailBuilder.HideProperty(NodeHandle->GetChildHandle(GET_MEMBER_NAME_CHECKED(FAnimNode_StrideWarping, LocomotionSpeed)));
-		DetailBuilder.HideProperty(NodeHandle->GetChildHandle(GET_MEMBER_NAME_CHECKED(FAnimNode_StrideWarping, MinSpeedTolerance)));
+		DetailBuilder.HideProperty(NodeHandle->GetChildHandle(GET_MEMBER_NAME_CHECKED(FAnimNode_StrideWarping, MinLocomotionSpeedThreshold)));
 	}
 }
 
@@ -102,6 +143,7 @@ void UAnimGraphNode_StrideWarping::PostEditChangeProperty(struct FPropertyChange
 			bRequiresNodeReconstruct = true;
 		}
 
+		// Evaluation mode
 		if (ChangedProperty->GetFName() == GET_MEMBER_NAME_STRING_CHECKED(FAnimNode_StrideWarping, Mode))
 		{
 			FScopedTransaction Transaction(LOCTEXT("ChangeEvaluationMode", "Change Evaluation Mode"));
@@ -111,14 +153,14 @@ void UAnimGraphNode_StrideWarping::PostEditChangeProperty(struct FPropertyChange
 			for (int32 PinIndex = 0; PinIndex < Pins.Num(); ++PinIndex)
 			{
 				UEdGraphPin* Pin = Pins[PinIndex];
-				if (Pin->PinName == GET_MEMBER_NAME_STRING_CHECKED(FAnimNode_StrideWarping, ManualStrideWarpingDir))
+				if (Pin->PinName == GET_MEMBER_NAME_STRING_CHECKED(FAnimNode_StrideWarping, StrideDirection))
 				{
 					if (Node.Mode == EWarpingEvaluationMode::Graph)
 					{
 						Pin->BreakAllPinLinks();
 					}
 				}
-				else if (Pin->PinName == GET_MEMBER_NAME_STRING_CHECKED(FAnimNode_StrideWarping, ManualStrideScaling))
+				else if (Pin->PinName == GET_MEMBER_NAME_STRING_CHECKED(FAnimNode_StrideWarping, StrideScale))
 				{
 					if (Node.Mode == EWarpingEvaluationMode::Graph)
 					{
@@ -132,7 +174,7 @@ void UAnimGraphNode_StrideWarping::PostEditChangeProperty(struct FPropertyChange
 						Pin->BreakAllPinLinks();
 					}
 				}
-				else if (Pin->PinName == GET_MEMBER_NAME_STRING_CHECKED(FAnimNode_StrideWarping, MinSpeedTolerance))
+				else if (Pin->PinName == GET_MEMBER_NAME_STRING_CHECKED(FAnimNode_StrideWarping, MinLocomotionSpeedThreshold))
 				{
 					if (Node.Mode == EWarpingEvaluationMode::Manual)
 					{
diff --git a/Engine/Plugins/Experimental/Animation/AnimationWarping/Source/Editor/Public/AnimGraph/AnimGraphNode_OrientationWarping.h b/Engine/Plugins/Experimental/Animation/AnimationWarping/Source/Editor/Public/AnimGraph/AnimGraphNode_OrientationWarping.h
--- a/Engine/Plugins/Experimental/Animation/AnimationWarping/Source/Editor/Public/AnimGraph/AnimGraphNode_OrientationWarping.h
+++ b/Engine/Plugins/Experimental/Animation/AnimationWarping/Source/Editor/Public/AnimGraph/AnimGraphNode_OrientationWarping.h
@@ -16,13 +16,19 @@ class ANIMATIONWARPINGEDITOR_API UAnimGraphNode_OrientationWarping : public UAni
 public:
 	// UEdGraphNode interface
 	virtual FText GetNodeTitle(ENodeTitleType::Type TitleType) const override;
+	virtual FLinearColor GetNodeTitleColor() const override;
 	virtual FText GetTooltipText() const override;
+	virtual void CustomizeDetails(IDetailLayoutBuilder& DetailBuilder) override;
 	virtual void PostEditChangeProperty(struct FPropertyChangedEvent& PropertyChangedEvent) override;
 	virtual void GetInputLinkAttributes(FNodeAttributeArray& OutAttributes) const override;
 	virtual void GetOutputLinkAttributes(FNodeAttributeArray& OutAttributes) const override;
 	// End of UEdGraphNode interface
 
 protected:
+	// UAnimGraphNode_Base interface
+	virtual void CustomizePinData(UEdGraphPin* Pin, FName SourcePropertyName, int32 ArrayIndex) const override;
+	// End of UAnimGraphNode_Base interface
+
 	// UAnimGraphNode_SkeletalControlBase interface
 	virtual FText GetControllerDescription() const override;
 	virtual const FAnimNode_SkeletalControlBase* GetNode() const override { return &Node; }
diff --git a/Engine/Plugins/Experimental/Animation/AnimationWarping/Source/Editor/Public/AnimGraph/AnimGraphNode_SlopeWarping.h b/Engine/Plugins/Experimental/Animation/AnimationWarping/Source/Editor/Public/AnimGraph/AnimGraphNode_SlopeWarping.h
--- a/Engine/Plugins/Experimental/Animation/AnimationWarping/Source/Editor/Public/AnimGraph/AnimGraphNode_SlopeWarping.h
+++ b/Engine/Plugins/Experimental/Animation/AnimationWarping/Source/Editor/Public/AnimGraph/AnimGraphNode_SlopeWarping.h
@@ -19,8 +19,7 @@ public:
 	// UEdGraphNode interface
 	virtual FText GetNodeTitle(ENodeTitleType::Type TitleType) const override;
 	virtual FText GetTooltipText() const override;
-	virtual void GetInputLinkAttributes(FNodeAttributeArray& OutAttributes) const override;
-	virtual void GetOutputLinkAttributes(FNodeAttributeArray& OutAttributes) const override;
+	virtual FLinearColor GetNodeTitleColor() const override;
 	// End of UEdGraphNode interface
 
 protected:
diff --git a/Engine/Plugins/Experimental/Animation/AnimationWarping/Source/Editor/Public/AnimGraph/AnimGraphNode_StrideWarping.h b/Engine/Plugins/Experimental/Animation/AnimationWarping/Source/Editor/Public/AnimGraph/AnimGraphNode_StrideWarping.h
--- a/Engine/Plugins/Experimental/Animation/AnimationWarping/Source/Editor/Public/AnimGraph/AnimGraphNode_StrideWarping.h
+++ b/Engine/Plugins/Experimental/Animation/AnimationWarping/Source/Editor/Public/AnimGraph/AnimGraphNode_StrideWarping.h
@@ -18,6 +18,7 @@ class ANIMATIONWARPINGEDITOR_API UAnimGraphNode_StrideWarping : public UAnimGrap
 public:
 	// UEdGraphNode interface
 	virtual FText GetNodeTitle(ENodeTitleType::Type TitleType) const override;
+	virtual FLinearColor GetNodeTitleColor() const override;
 	virtual FText GetTooltipText() const override;
 	virtual void CustomizeDetails(IDetailLayoutBuilder& DetailBuilder) override;
 	virtual void PostEditChangeProperty(struct FPropertyChangedEvent& PropertyChangedEvent) override;
diff --git a/Engine/Plugins/Experimental/Animation/AnimationWarping/Source/Runtime/AnimationWarpingRuntime.Build.cs b/Engine/Plugins/Experimental/Animation/AnimationWarping/Source/Runtime/AnimationWarpingRuntime.Build.cs
--- a/Engine/Plugins/Experimental/Animation/AnimationWarping/Source/Runtime/AnimationWarpingRuntime.Build.cs
+++ b/Engine/Plugins/Experimental/Animation/AnimationWarping/Source/Runtime/AnimationWarpingRuntime.Build.cs
@@ -9,18 +9,12 @@ namespace UnrealBuildTool.Rules
 			PublicDependencyModuleNames.AddRange(
 			new string[]
 			{
+				"AnimationCore",
                 "AnimGraphRuntime",
-                "AnimationCore",
                 "Core",
                 "CoreUObject",
                 "Engine",
             });
-
-			PrivateDependencyModuleNames.AddRange(
-			new string[]
-			{
-                "SlateCore",
-            });
 		}
 	}
 }
diff --git a/Engine/Plugins/Experimental/Animation/AnimationWarping/Source/Runtime/Private/BoneControllers/AnimNode_OrientationWarping.cpp b/Engine/Plugins/Experimental/Animation/AnimationWarping/Source/Runtime/Private/BoneControllers/AnimNode_OrientationWarping.cpp
--- a/Engine/Plugins/Experimental/Animation/AnimationWarping/Source/Runtime/Private/BoneControllers/AnimNode_OrientationWarping.cpp
+++ b/Engine/Plugins/Experimental/Animation/AnimationWarping/Source/Runtime/Private/BoneControllers/AnimNode_OrientationWarping.cpp
@@ -1,47 +1,69 @@
 // Copyright Epic Games, Inc. All Rights Reserved.
 
 #include "BoneControllers/AnimNode_OrientationWarping.h"
-#include "Animation/AnimRootMotionProvider.h"
-#include "AnimationRuntime.h"
 #include "Animation/AnimInstanceProxy.h"
+#include "Animation/AnimRootMotionProvider.h"
 
 DECLARE_CYCLE_STAT(TEXT("OrientationWarping Eval"), STAT_OrientationWarping_Eval, STATGROUP_Anim);
 
 #if ENABLE_ANIM_DEBUG
-TAutoConsoleVariable<int32> CVarAnimNodeOrientationWarpingDebug(TEXT("a.AnimNode.OrientationWarping.Debug"), 0, TEXT("Turn on debug for AnimNode_OrientationWarping"));
+TAutoConsoleVariable<int32> CVarAnimNodeOrientationWarpingDebug(TEXT("a.AnimNode.OrientationWarping.Debug"), 0, TEXT("Turn on visualization debugging for Orientation Warping"));
+TAutoConsoleVariable<int32> CVarAnimNodeOrientationWarpingVerbose(TEXT("a.AnimNode.OrientationWarping.Verbose"), 0, TEXT("Turn on verbose graph debugging for Orientation Warping"));
 TAutoConsoleVariable<int32> CVarAnimNodeOrientationWarpingEnable(TEXT("a.AnimNode.OrientationWarping.Enable"), 1, TEXT("Toggle Orientation Warping"));
 #endif
 
-FAnimNode_OrientationWarping::FAnimNode_OrientationWarping()
-	: Mode(EWarpingEvaluationMode::Manual)
-	, LocomotionAngle(0.f)
-	, IKFootRootBoneIndex(INDEX_NONE)
-{
-}
-
-static inline FVector GetAxisVector(const EAxis::Type& InAxis)
+namespace UE::Anim
 {
-	switch (InAxis)
+	static inline FVector GetAxisVector(const EAxis::Type& InAxis)
 	{
-	case EAxis::X: return FVector::ForwardVector;
-	case EAxis::Y: return FVector::RightVector;
-	default:
-	case EAxis::Z: return FVector::UpVector;
-	};
-}
+		switch (InAxis)
+		{
+		case EAxis::X:
+			return FVector::ForwardVector;
+		case EAxis::Y:
+			return FVector::RightVector;
+		default:
+			return FVector::UpVector;
+		};
+	}
 
-static inline bool IsInvalidWarpingAngle(float Angle, float Tolerance)
-{
-	return FMath::IsNearlyZero(Angle, Tolerance) || FMath::IsNearlyEqual(FMath::Abs(Angle), PI, Tolerance);
+	static inline bool IsInvalidWarpingAngle(float Angle, float Tolerance)
+	{
+		return FMath::IsNearlyZero(Angle, Tolerance) || FMath::IsNearlyEqual(FMath::Abs(Angle), PI, Tolerance);
+	}
 }
 
 void FAnimNode_OrientationWarping::GatherDebugData(FNodeDebugData& DebugData)
 {
 	FString DebugLine = DebugData.GetNodeName(this);
- 	DebugLine += FString::Printf(TEXT("Angle(%.1fd)"), LocomotionAngle);
-
+#if ENABLE_ANIM_DEBUG
+	if (CVarAnimNodeOrientationWarpingVerbose.GetValueOnAnyThread() == 1)
+	{
+		if (Mode == EWarpingEvaluationMode::Manual)
+		{
+			DebugLine += TEXT("\n - Evaluation Mode: (Manual)");
+			DebugLine += FString::Printf(TEXT("\n - Orientation Angle: (%.3fd)"), FMath::RadiansToDegrees(OrientationAngle));
+		}
+		else
+		{
+			DebugLine += TEXT("\n - Evaluation Mode: (Graph)");
+			DebugLine += FString::Printf(TEXT("\n - Orientation Angle: (%.3fd)"), FMath::RadiansToDegrees(OrientationAngle));
+			DebugLine += FString::Printf(TEXT("\n - Locomotion Angle: (%.3fd)"), LocomotionAngle);
+			DebugLine += FString::Printf(TEXT("\n - Locomotion Delta Angle Threshold: (%.3fd)"), LocomotionAngleDeltaThreshold);
+		}
+		DebugLine += FString::Printf(TEXT("\n - Distributed Bone Orientation Alpha: (%.3fd)"), DistributedBoneOrientationAlpha);
+		if (const UEnum* TypeEnum = FindObject<UEnum>(ANY_PACKAGE, TEXT("EAxis")))
+		{
+			DebugLine += FString::Printf(TEXT("\n - Rotation Axis: (%s)"), *(TypeEnum->GetNameStringByIndex(static_cast<int32>(RotationAxis))));
+		}
+		DebugLine += FString::Printf(TEXT("\n - Rotation Interpolation Speed: (%.3fd)"), RotationInterpSpeed);
+	}
+	else
+#endif
+	{
+		DebugLine += FString::Printf(TEXT("(Orientation Angle: %.3fd)"), FMath::RadiansToDegrees(OrientationAngle));
+	}
 	DebugData.AddDebugItem(DebugLine);
-
 	ComponentPose.GatherDebugData(DebugData);
 }
 
@@ -58,9 +80,11 @@ void FAnimNode_OrientationWarping::UpdateInternal(const FAnimationUpdateContext&
 void FAnimNode_OrientationWarping::EvaluateSkeletalControl_AnyThread(FComponentSpacePoseContext& Output, TArray<FBoneTransform>& OutBoneTransforms)
 {
 	SCOPE_CYCLE_COUNTER(STAT_OrientationWarping_Eval);
-
 	check(OutBoneTransforms.Num() == 0);
 
+	const UE::Anim::IAnimRootMotionProvider* RootMotionProvider = UE::Anim::IAnimRootMotionProvider::Get();
+	const bool bGraphDrivenWarping = RootMotionProvider && Mode == EWarpingEvaluationMode::Graph;
+
 	// We will likely need to revisit LocomotionAngle participating as an input to orientation warping.
 	// Without velocity information from the motion model (such as the capsule), LocomotionAngle isn't enough
 	// information in isolation for all cases when deciding to warp.
@@ -74,11 +98,8 @@ void FAnimNode_OrientationWarping::EvaluateSkeletalControl_AnyThread(FComponentS
 	// The solution may be instead to pass velocity with the actor base rotation, allowing us to retain
 	// speed information about the motion. It may also allow us to do more complex orienting behavior 
 	// when multiple degrees of freedom can be considered.
-	float LocomotionRotationAngle = FMath::DegreesToRadians(FRotator::NormalizeAxis(LocomotionAngle));
-
-	const FVector LocomotionRotationAxis = GetAxisVector(Settings.YawRotationAxis);
-	const UE::Anim::IAnimRootMotionProvider* RootMotionProvider = UE::Anim::IAnimRootMotionProvider::Get();
-	const bool bGraphDrivenWarping = RootMotionProvider && Mode == EWarpingEvaluationMode::Graph;
+	OrientationAngle = FMath::DegreesToRadians(FRotator::NormalizeAxis(bGraphDrivenWarping ? LocomotionAngle : OrientationAngle));
+	const FVector RotationAxisVector = UE::Anim::GetAxisVector(RotationAxis);
 
 	// Graph driven orientation warping will modify the incoming root motion to orient towards the intended locomotion angle
 	if (bGraphDrivenWarping)
@@ -88,7 +109,7 @@ void FAnimNode_OrientationWarping::EvaluateSkeletalControl_AnyThread(FComponentS
 
 		if (ensure(bRootMotionDeltaPresent))
 		{
-			// In UE, forward is defined as +x; consequently this is also true when sampling an actor's velocity. Sometimes the skeletal 
+			// In UE, forward is defined as +x; consequently this is also true when sampling an actor's velocity. Historically the skeletal 
 			// mesh component forward will not match the actor, requiring us to correct the rotation before sampling the LocomotionForward.
 			// In order to make orientation warping 'pure' in the future we will need to provide more context about the intent of
 			// the actor vs the intent of the animation in their respective spaces. Specifically, we will need some form the following information:
@@ -98,22 +119,22 @@ void FAnimNode_OrientationWarping::EvaluateSkeletalControl_AnyThread(FComponentS
 			// 3. Skeletal Mesh Relative Rotation
 			const FTransform SkeletalMeshRelativeTransform = Output.AnimInstanceProxy->GetComponentRelativeTransform();
 			const FQuat SkeletalMeshRelativeRotation = SkeletalMeshRelativeTransform.GetRotation();
-			const FQuat LocomotionRotation = FQuat(LocomotionRotationAxis, LocomotionRotationAngle);
+			const FQuat LocomotionRotation = FQuat(RotationAxisVector, OrientationAngle);
 			const FVector LocomotionForwardDir = SkeletalMeshRelativeRotation.UnrotateVector(LocomotionRotation.GetForwardVector());
 			const FVector RootMotionDeltaDir = RootMotionTransformDelta.GetTranslation();
 
 			// Capture the delta rotation from the axis of motion we care about
 			FQuat WarpedRotation = FQuat::FindBetween(RootMotionDeltaDir, LocomotionForwardDir);
-			LocomotionRotationAngle = WarpedRotation.GetTwistAngle(LocomotionRotationAxis);
+			OrientationAngle = WarpedRotation.GetTwistAngle(RotationAxisVector);
 
 			// Motion Matching may return an animation that deviates a lot from the movement direction (e.g movement direction going bwd and motion matching could return the fwd animation for a few frames)
 			// When that happens, since we use the delta between root motion and movement direction, we would be over-rotating the lower body and breaking the pose during those frames
 			// So, when that happens we use the inverse of the movement direction to calculate our target rotation. 
 			// This feels a bit 'hacky' but its the only option I've found so far to mitigate the problem
-			if (LocomotionInversionThresholdAngle > 0.f && FMath::Abs(FMath::RadiansToDegrees(LocomotionRotationAngle)) > LocomotionInversionThresholdAngle)
+			if (LocomotionAngleDeltaThreshold > 0.f && FMath::Abs(FMath::RadiansToDegrees(OrientationAngle)) > LocomotionAngleDeltaThreshold)
 			{
 				WarpedRotation = FQuat::FindBetween(RootMotionDeltaDir, -LocomotionForwardDir);
-				LocomotionRotationAngle = WarpedRotation.GetTwistAngle(LocomotionRotationAxis);
+				OrientationAngle = WarpedRotation.GetTwistAngle(RotationAxisVector);
 			}
 
 			// Rotate the root motion delta fully by the warped angle
@@ -124,48 +145,49 @@ void FAnimNode_OrientationWarping::EvaluateSkeletalControl_AnyThread(FComponentS
 			// Forward the side effects of orientation warping on the root motion contribution for this sub-graph
 			const bool bRootMotionOverridden = RootMotionProvider->OverrideRootMotion(RootMotionTransformDelta, Output.CustomAttributes);
 			ensure(bRootMotionOverridden);
-
 #if ENABLE_ANIM_DEBUG
-			const bool bShowDebug = (CVarAnimNodeOrientationWarpingDebug.GetValueOnAnyThread() == 1);
-			if (bShowDebug)
+			bool bDebugging = false;
+#if WITH_EDITORONLY_DATA
+			bDebugging = bDebugging || bEnableDebugDraw;
+#endif
+			bDebugging = bDebugging || CVarAnimNodeOrientationWarpingDebug.GetValueOnAnyThread() == 1;
+
+			if (bDebugging)
 			{
-				const FVector LocomotionDirWS = Output.AnimInstanceProxy->GetActorTransform().TransformVectorNoScale(LocomotionRotation.GetForwardVector()).GetSafeNormal();
-				const FVector RootMotionDirWS = Output.AnimInstanceProxy->GetActorTransform().TransformVectorNoScale(SkeletalMeshRelativeRotation.RotateVector(RootMotionDeltaDir)).GetSafeNormal();
+				const FVector LocomotionDir = Output.AnimInstanceProxy->GetActorTransform().TransformVectorNoScale(LocomotionRotation.GetForwardVector()).GetSafeNormal();
+				const FVector RootMotionDir = Output.AnimInstanceProxy->GetActorTransform().TransformVectorNoScale(SkeletalMeshRelativeRotation.RotateVector(RootMotionDeltaDir)).GetSafeNormal();
 
-				const float DrawDebugArrowScale = 100.f;
-				Output.AnimInstanceProxy->AnimDrawDebugLine(
+				Output.AnimInstanceProxy->AnimDrawDebugDirectionalArrow(
 					Output.AnimInstanceProxy->GetComponentTransform().GetLocation(),
-					Output.AnimInstanceProxy->GetComponentTransform().GetLocation() + LocomotionDirWS * DrawDebugArrowScale,
-					FColor::Red, false, 0.f, 1.f);
+					Output.AnimInstanceProxy->GetComponentTransform().GetLocation() + LocomotionDir * 100.f,
+					40.f, FColor::Red, false, 0.f, 2.f);
 
-				Output.AnimInstanceProxy->AnimDrawDebugLine(
+				Output.AnimInstanceProxy->AnimDrawDebugDirectionalArrow(
 					Output.AnimInstanceProxy->GetComponentTransform().GetLocation(),
-					Output.AnimInstanceProxy->GetComponentTransform().GetLocation() + RootMotionDirWS * DrawDebugArrowScale,
-					FColor::Blue, false, 0.f, 1.f);
+					Output.AnimInstanceProxy->GetComponentTransform().GetLocation() + RootMotionDir * 100.f,
+					40.f, FColor::Blue, false, 0.f, 2.f);
 			}
 #endif
 		}
 	} 
-	else if (IsInvalidWarpingAngle(LocomotionRotationAngle, KINDA_SMALL_NUMBER))
+	else if (UE::Anim::IsInvalidWarpingAngle(OrientationAngle, KINDA_SMALL_NUMBER))
 	{
 		return;
 	}
 
-	if (InterpSpeed > 0.f)
+	if (RotationInterpSpeed > 0.f)
 	{
-		LocomotionRotationAngle = FMath::FInterpTo(LastRotationAngle, LocomotionRotationAngle, Output.AnimInstanceProxy->GetDeltaSeconds(), InterpSpeed);
-		LastRotationAngle = LocomotionRotationAngle;
+		OrientationAngle = FMath::FInterpTo(PreviousWarpedRotation, OrientationAngle, Output.AnimInstanceProxy->GetDeltaSeconds(), RotationInterpSpeed);
+		PreviousWarpedRotation = OrientationAngle;
 	}
 
 	// Allow the alpha value of the node to affect the final rotation
-	LocomotionRotationAngle *= ActualAlpha;
-
-	const float BodyOrientationAlpha = FMath::Clamp(Settings.BodyOrientationAlpha, 0.f, 1.f);
+	OrientationAngle *= ActualAlpha;
 
 	// Rotate Root Bone first, as that cheaply rotates the whole pose with one transformation.
-	if (!FMath::IsNearlyZero(BodyOrientationAlpha, KINDA_SMALL_NUMBER))
+	if (!FMath::IsNearlyZero(DistributedBoneOrientationAlpha, KINDA_SMALL_NUMBER))
 	{
-		const FQuat RootRotation = FQuat(LocomotionRotationAxis, LocomotionRotationAngle * BodyOrientationAlpha);
+		const FQuat RootRotation = FQuat(RotationAxisVector, OrientationAngle * DistributedBoneOrientationAlpha);
 		const FCompactPoseBoneIndex RootBoneIndex(0);
 
 		FTransform RootBoneTransform(Output.Pose.GetComponentSpaceTransform(RootBoneIndex));
@@ -175,8 +197,8 @@ void FAnimNode_OrientationWarping::EvaluateSkeletalControl_AnyThread(FComponentS
 	}
 		
 	const int32 NumSpineBones = SpineBoneDataArray.Num();
-	const bool bBodyOrientationAlpha = !FMath::IsNearlyZero(BodyOrientationAlpha, KINDA_SMALL_NUMBER);
-	const bool bUpdateSpineBones = (NumSpineBones > 0) && bBodyOrientationAlpha;
+	const bool bSpineOrientationAlpha = !FMath::IsNearlyZero(DistributedBoneOrientationAlpha, KINDA_SMALL_NUMBER);
+	const bool bUpdateSpineBones = (NumSpineBones > 0) && bSpineOrientationAlpha;
 
 	if (bUpdateSpineBones)
 	{
@@ -184,7 +206,7 @@ void FAnimNode_OrientationWarping::EvaluateSkeletalControl_AnyThread(FComponentS
 		for (int32 ArrayIndex = 0; ArrayIndex < NumSpineBones; ArrayIndex++)
 		{
 			const FOrientationWarpingSpineBoneData& BoneData = SpineBoneDataArray[ArrayIndex];
-			const FQuat SpineBoneCounterRotation = FQuat(LocomotionRotationAxis, -LocomotionRotationAngle * BodyOrientationAlpha * BoneData.Weight);
+			const FQuat SpineBoneCounterRotation = FQuat(RotationAxisVector, -OrientationAngle * DistributedBoneOrientationAlpha * BoneData.Weight);
 			check(BoneData.Weight > 0.f);
 
 			FTransform SpineBoneTransform(Output.Pose.GetComponentSpaceTransform(BoneData.BoneIndex));
@@ -194,32 +216,32 @@ void FAnimNode_OrientationWarping::EvaluateSkeletalControl_AnyThread(FComponentS
 		}
 	}
 
-	const float IKFootRootOrientationAlpha = 1.f - BodyOrientationAlpha;
-	const bool bUpdateIKFootRoot = (IKFootRootBoneIndex != FCompactPoseBoneIndex(INDEX_NONE)) && !FMath::IsNearlyZero(IKFootRootOrientationAlpha, KINDA_SMALL_NUMBER);
+	const float IKFootRootOrientationAlpha = 1.f - DistributedBoneOrientationAlpha;
+	const bool bUpdateIKFootRoot = (IKFootData.IKFootRootBoneIndex != FCompactPoseBoneIndex(INDEX_NONE)) && !FMath::IsNearlyZero(IKFootRootOrientationAlpha, KINDA_SMALL_NUMBER);
 
 	// Rotate IK Foot Root
 	if (bUpdateIKFootRoot)
 	{
-		const FQuat BoneRotation = FQuat(LocomotionRotationAxis, LocomotionRotationAngle * IKFootRootOrientationAlpha);
+		const FQuat BoneRotation = FQuat(RotationAxisVector, OrientationAngle * IKFootRootOrientationAlpha);
 
-		FTransform IKFootRootTransform(Output.Pose.GetComponentSpaceTransform(IKFootRootBoneIndex));
+		FTransform IKFootRootTransform(Output.Pose.GetComponentSpaceTransform(IKFootData.IKFootRootBoneIndex));
 		IKFootRootTransform.SetRotation(BoneRotation * IKFootRootTransform.GetRotation());
 		IKFootRootTransform.NormalizeRotation();
-		Output.Pose.SetComponentSpaceTransform(IKFootRootBoneIndex, IKFootRootTransform);
+		Output.Pose.SetComponentSpaceTransform(IKFootData.IKFootRootBoneIndex, IKFootRootTransform);
 
 		// IK Feet 
 		// These match the root orientation, so don't rotate them. Just preserve root rotation. 
 		// We need to update their translation though, since we rotated their parent (the IK Foot Root bone).
-		const int32 NumIKFootBones = IKFootBoneIndexArray.Num();
+		const int32 NumIKFootBones = IKFootData.IKFootBoneIndexArray.Num();
 		const bool bUpdateIKFootBones = bUpdateIKFootRoot && (NumIKFootBones > 0);
 
 		if (bUpdateIKFootBones)
 		{
-			const FQuat IKFootRotation = FQuat(LocomotionRotationAxis, -LocomotionRotationAngle * IKFootRootOrientationAlpha);
+			const FQuat IKFootRotation = FQuat(RotationAxisVector, -OrientationAngle * IKFootRootOrientationAlpha);
 
 			for (int32 ArrayIndex = 0; ArrayIndex < NumIKFootBones; ArrayIndex++)
 			{
-				const FCompactPoseBoneIndex& IKFootBoneIndex = IKFootBoneIndexArray[ArrayIndex];
+				const FCompactPoseBoneIndex& IKFootBoneIndex = IKFootData.IKFootBoneIndexArray[ArrayIndex];
 
 				FTransform IKFootBoneTransform(Output.Pose.GetComponentSpaceTransform(IKFootBoneIndex));
 				IKFootBoneTransform.SetRotation(IKFootRotation * IKFootBoneTransform.GetRotation());
@@ -240,91 +262,110 @@ bool FAnimNode_OrientationWarping::IsValidToEvaluate(const USkeleton* Skeleton,
 		return false;
 	}
 #endif
+	if (RotationAxis == EAxis::None)
+	{
+		return false;
+	}
 
-	bool bIKFootRootIsValid = IKFootRootBoneIndex != INDEX_NONE;
-	bool bIKFeetAreValid = IKFootBoneIndexArray.Num() > 0;
-	for (const auto& IKFootBoneIndex : IKFootBoneIndexArray)
+	if (SpineBoneDataArray.IsEmpty())
 	{
-		bIKFeetAreValid = bIKFeetAreValid && IKFootBoneIndex != INDEX_NONE;
+		return false;
+	}
+	else
+	{
+		for (const auto& Spine : SpineBoneDataArray)
+		{
+			if (Spine.BoneIndex == INDEX_NONE)
+			{
+				return false;
+			}
+		}
 	}
 
-	bool bSpineIsValid = SpineBoneDataArray.Num() > 0;
-	for (const auto& Spine : SpineBoneDataArray)
+	if (IKFootData.IKFootRootBoneIndex == INDEX_NONE)
 	{
-		bSpineIsValid = bSpineIsValid && Spine.BoneIndex != INDEX_NONE;
+		return false;
 	}
 
-	return bIKFootRootIsValid && bIKFeetAreValid && bSpineIsValid;
+	if (IKFootData.IKFootBoneIndexArray.IsEmpty())
+	{
+		return false;
+	}
+	else
+	{
+		for (const auto& IKFootBoneIndex : IKFootData.IKFootBoneIndexArray)
+		{
+			if (IKFootBoneIndex == INDEX_NONE)
+			{
+				return false;
+			}
+		}
+	}
+	return true;
 }
 
 void FAnimNode_OrientationWarping::InitializeBoneReferences(const FBoneContainer& RequiredBones)
 {
-	Settings.IKFootRootBone.Initialize(RequiredBones);
-	IKFootRootBoneIndex = Settings.IKFootRootBone.GetCompactPoseIndex(RequiredBones);
+	IKFootRootBone.Initialize(RequiredBones);
+	IKFootData.IKFootRootBoneIndex = IKFootRootBone.GetCompactPoseIndex(RequiredBones);
 
+	IKFootData.IKFootBoneIndexArray.Reset();
+	for (auto& Bone : IKFootBones)
 	{
-		IKFootBoneIndexArray.Reset();
-		for (auto& BoneSettings : Settings.IKFootBones)
-		{
-			BoneSettings.Initialize(RequiredBones);
-			IKFootBoneIndexArray.Add(BoneSettings.GetCompactPoseIndex(RequiredBones));
-		}
+		Bone.Initialize(RequiredBones);
+		IKFootData.IKFootBoneIndexArray.Add(Bone.GetCompactPoseIndex(RequiredBones));
 	}
 
+	SpineBoneDataArray.Reset();
+	for (auto& Bone : SpineBones)
 	{
-		SpineBoneDataArray.Reset();
-		for (auto& BoneSettings : Settings.SpineBones)
-		{
-			BoneSettings.Bone.Initialize(RequiredBones);
-			SpineBoneDataArray.Add(FOrientationWarpingSpineBoneData(BoneSettings.Bone.GetCompactPoseIndex(RequiredBones)));
-		}
+		Bone.Initialize(RequiredBones);
+		SpineBoneDataArray.Add(FOrientationWarpingSpineBoneData(Bone.GetCompactPoseIndex(RequiredBones)));
+	}
 
-		if (SpineBoneDataArray.Num() > 0)
-		{
-			// Sort bones indices so we can transform parent before child
-			SpineBoneDataArray.Sort(FOrientationWarpingSpineBoneData::FCompareBoneIndex());
+	if (SpineBoneDataArray.Num() > 0)
+	{
+		// Sort bones indices so we can transform parent before child
+		SpineBoneDataArray.Sort(FOrientationWarpingSpineBoneData::FCompareBoneIndex());
 
-			// Assign Weights.
+		// Assign Weights.
+		TArray<int32, TInlineAllocator<20>> IndicesToUpdate;
+
+		for (int32 Index = SpineBoneDataArray.Num() - 1; Index >= 0; Index--)
+		{
+			// If this bone's weight hasn't been updated, scan his parents.
+			// If parents have weight, we add it to 'ExistingWeight'.
+			// split (1.f - 'ExistingWeight') between all members of the chain that have no weight yet.
+			if (SpineBoneDataArray[Index].Weight == 0.f)
 			{
-				TArray<int32> IndicesToUpdate;
+				IndicesToUpdate.Reset(SpineBoneDataArray.Num());
+				float ExistingWeight = 0.f;
+				IndicesToUpdate.Add(Index);
 
-				for (int32 Index = SpineBoneDataArray.Num() - 1; Index >= 0; Index--)
+				const FCompactPoseBoneIndex CompactBoneIndex = SpineBoneDataArray[Index].BoneIndex;
+				for (int32 ParentIndex = Index - 1; ParentIndex >= 0; ParentIndex--)
 				{
-					// If this bone's weight hasn't been updated, scan his parents.
-					// If parents have weight, we add it to 'ExistingWeight'.
-					// split (1.f - 'ExistingWeight') between all members of the chain that have no weight yet.
-					if (SpineBoneDataArray[Index].Weight == 0.f)
+					if (RequiredBones.BoneIsChildOf(CompactBoneIndex, SpineBoneDataArray[ParentIndex].BoneIndex))
 					{
-						IndicesToUpdate.Reset(SpineBoneDataArray.Num());
-						float ExistingWeight = 0.f;
-						IndicesToUpdate.Add(Index);
-
-						const FCompactPoseBoneIndex CompactBoneIndex = SpineBoneDataArray[Index].BoneIndex;
-						for (int32 ParentIndex = Index - 1; ParentIndex >= 0; ParentIndex--)
+						if (SpineBoneDataArray[ParentIndex].Weight > 0.f)
 						{
-							if (RequiredBones.BoneIsChildOf(CompactBoneIndex, SpineBoneDataArray[ParentIndex].BoneIndex))
-							{
-								if (SpineBoneDataArray[ParentIndex].Weight > 0.f)
-								{
-									ExistingWeight += SpineBoneDataArray[ParentIndex].Weight;
-								}
-								else
-								{
-									IndicesToUpdate.Add(ParentIndex);
-								}
-							}
+							ExistingWeight += SpineBoneDataArray[ParentIndex].Weight;
 						}
-
-						check(IndicesToUpdate.Num() > 0);
-						const float WeightToShare = 1.f - ExistingWeight;
-						const float IndividualWeight = WeightToShare / float(IndicesToUpdate.Num());
-
-						for (int32 UpdateListIndex = 0; UpdateListIndex < IndicesToUpdate.Num(); UpdateListIndex++)
+						else
 						{
-							SpineBoneDataArray[IndicesToUpdate[UpdateListIndex]].Weight = IndividualWeight;
+							IndicesToUpdate.Add(ParentIndex);
 						}
 					}
 				}
+
+				check(IndicesToUpdate.Num() > 0);
+				const float WeightToShare = 1.f - ExistingWeight;
+				const float IndividualWeight = WeightToShare / float(IndicesToUpdate.Num());
+
+				for (int32 UpdateListIndex = 0; UpdateListIndex < IndicesToUpdate.Num(); UpdateListIndex++)
+				{
+					SpineBoneDataArray[IndicesToUpdate[UpdateListIndex]].Weight = IndividualWeight;
+				}
 			}
 		}
 	}
diff --git a/Engine/Plugins/Experimental/Animation/AnimationWarping/Source/Runtime/Private/BoneControllers/AnimNode_SlopeWarping.cpp b/Engine/Plugins/Experimental/Animation/AnimationWarping/Source/Runtime/Private/BoneControllers/AnimNode_SlopeWarping.cpp
--- a/Engine/Plugins/Experimental/Animation/AnimationWarping/Source/Runtime/Private/BoneControllers/AnimNode_SlopeWarping.cpp
+++ b/Engine/Plugins/Experimental/Animation/AnimationWarping/Source/Runtime/Private/BoneControllers/AnimNode_SlopeWarping.cpp
@@ -21,7 +21,6 @@ FAnimNode_SlopeWarping::FAnimNode_SlopeWarping()
 	: MySkelMeshComponent(nullptr)
 	, MyMovementComponent(nullptr)
 	, MyAnimInstanceProxy(nullptr)
-	, Mode(EWarpingEvaluationMode::Manual)
 	, GravityDir(-FVector::UpVector)
 	, CustomFloorOffset(FVector::ZeroVector)
 	, CachedDeltaTime(0.f)
diff --git a/Engine/Plugins/Experimental/Animation/AnimationWarping/Source/Runtime/Private/BoneControllers/AnimNode_StrideWarping.cpp b/Engine/Plugins/Experimental/Animation/AnimationWarping/Source/Runtime/Private/BoneControllers/AnimNode_StrideWarping.cpp
--- a/Engine/Plugins/Experimental/Animation/AnimationWarping/Source/Runtime/Private/BoneControllers/AnimNode_StrideWarping.cpp
+++ b/Engine/Plugins/Experimental/Animation/AnimationWarping/Source/Runtime/Private/BoneControllers/AnimNode_StrideWarping.cpp
@@ -1,95 +1,61 @@
 // Copyright Epic Games, Inc. All Rights Reserved.
 
 #include "BoneControllers/AnimNode_StrideWarping.h"
-#include "AnimationRuntime.h"
 #include "Animation/AnimInstanceProxy.h"
-#include "DrawDebugHelpers.h"
-#include "Animation/InputScaleBias.h"
 #include "Animation/AnimRootMotionProvider.h"
 
+DECLARE_CYCLE_STAT(TEXT("StrideWarping Eval"), STAT_StrideWarping_Eval, STATGROUP_Anim);
+
 #if ENABLE_ANIM_DEBUG
-TAutoConsoleVariable<int32> CVarAnimNodeStrideWarpingDebug(TEXT("a.AnimNode.StrideWarping.Debug"), 0, TEXT("Turn on debug for AnimNode_StrideWarping"));
+TAutoConsoleVariable<int32> CVarAnimNodeStrideWarpingDebug(TEXT("a.AnimNode.StrideWarping.Debug"), 0, TEXT("Turn on visualization debugging for Stride Warping"));
+TAutoConsoleVariable<int32> CVarAnimNodeStrideWarpingVerbose(TEXT("a.AnimNode.StrideWarping.Verbose"), 0, TEXT("Turn on verbose graph debugging for Stride Warping"));
 TAutoConsoleVariable<int32> CVarAnimNodeStrideWarpingEnable(TEXT("a.AnimNode.StrideWarping.Enable"), 1, TEXT("Toggle Stride Warping"));
 #endif
 
-DECLARE_CYCLE_STAT(TEXT("StrideWarping Eval"), STAT_StrideWarping_Eval, STATGROUP_Anim);
-
-FAnimNode_StrideWarping::FAnimNode_StrideWarping()
-	: Mode(EWarpingEvaluationMode::Manual)
-	, ManualStrideWarpingDir(ForceInitToZero)
-	, ManualStrideScaling(1.f)
-	, LocomotionSpeed(0.f)
-	, MinSpeedTolerance(10.f)
-	, StrideWarpingAxisMode(EStrideWarpingAxisMode::ActorSpaceVectorInput)
-	, FloorNormalAxisMode(EStrideWarpingAxisMode::IKFootRootLocalZ)
-	, GravityDirAxisMode(EStrideWarpingAxisMode::ComponentSpaceVectorInput)
-	, ManualFloorNormalInput(ForceInitToZero)
-	, ManualGravityDirInput(-FVector::UpVector)
-	, PelvisPostAdjustmentAlpha(0.4f)
-	, PelvisAdjustmentMaxIter(3)
-	, bAdjustThighBonesRotation(true)
-	, bClampIKUsingFKLeg(true)
-	, bOrientStrideWarpingAxisBasedOnFloorNormal(true)
-	, CachedDeltaTime(0.f)
+void FAnimNode_StrideWarping::GatherDebugData(FNodeDebugData& DebugData)
 {
+	FString DebugLine = DebugData.GetNodeName(this);
+#if ENABLE_ANIM_DEBUG
+	if (CVarAnimNodeStrideWarpingVerbose.GetValueOnAnyThread() == 1)
+	{
+		DebugLine += FString::Printf(TEXT("\n - Evaluation Mode: (%s)"), Mode == EWarpingEvaluationMode::Graph ? TEXT("Graph") : TEXT("Manual"));
+		DebugLine += FString::Printf(TEXT("\n - Stride Scale: (%.3fd)"), StrideScale);
+		DebugLine += FString::Printf(TEXT("\n - Stride Direction: (%s)"), *(StrideDirection.ToCompactString()));
+		if (Mode == EWarpingEvaluationMode::Graph)
+		{
+			DebugLine += FString::Printf(TEXT("\n - Locomotion Speed: (%.3fd)"), LocomotionSpeed);
+			DebugLine += FString::Printf(TEXT("\n - Min Locomotion Speed Threshold: (%.3fd)"), MinLocomotionSpeedThreshold);
+		}
+		DebugLine += FString::Printf(TEXT("\n - Floor Normal: (%s)"), *(FloorNormalDirection.Value.ToCompactString()));
+		DebugLine += FString::Printf(TEXT("\n - Gravity Direction: (%s)"), *(GravityDirection.Value.ToCompactString()));
+	}
+	else
+#endif
+	{
+		DebugLine += FString::Printf(TEXT("(Stride Scale: %.3fd, Direction: %s)"), StrideScale, *(StrideDirection.ToCompactString()));
+	}
+	DebugData.AddDebugItem(DebugLine);
+	ComponentPose.GatherDebugData(DebugData);
 }
 
 void FAnimNode_StrideWarping::Initialize_AnyThread(const FAnimationInitializeContext& Context)
 {
 	FAnimNode_SkeletalControlBase::Initialize_AnyThread(Context);
-
-	MyAnimInstanceProxy = Context.AnimInstanceProxy;
-	StrideScalingScaleBiasClamp.Reinitialize();
-}
-
-void FAnimNode_StrideWarping::GatherDebugData(FNodeDebugData& DebugData)
-{
-	FString DebugLine = DebugData.GetNodeName(this);
-
-// 	DebugLine += "(";
-// 	AddDebugNodeData(DebugLine);
-// 	DebugLine += FString::Printf(TEXT(" Target: %s)"), *BoneToModify.BoneName.ToString());
-
-	DebugData.AddDebugItem(DebugLine);
-
-	ComponentPose.GatherDebugData(DebugData);
+	AnimInstanceProxy = Context.AnimInstanceProxy;
+	StrideScaleModifierState.Reinitialize();
 }
 
 void FAnimNode_StrideWarping::UpdateInternal(const FAnimationUpdateContext& Context)
 {
 	FAnimNode_SkeletalControlBase::UpdateInternal(Context);
-
 	CachedDeltaTime = Context.GetDeltaTime();
 }
 
-FVector FAnimNode_StrideWarping::GetAxisModeValue(const EStrideWarpingAxisMode& AxisMode, const FTransform& IKFootRootCSTransform, const FVector& UserSuppliedVector) const
-{
-	switch (AxisMode)
-	{
-	case EStrideWarpingAxisMode::IKFootRootLocalX: return IKFootRootCSTransform.GetUnitAxis(EAxis::X);
-	case EStrideWarpingAxisMode::IKFootRootLocalY: return IKFootRootCSTransform.GetUnitAxis(EAxis::Y);
-	case EStrideWarpingAxisMode::IKFootRootLocalZ: return IKFootRootCSTransform.GetUnitAxis(EAxis::Z);
-	case EStrideWarpingAxisMode::ComponentSpaceVectorInput: return UserSuppliedVector.GetSafeNormal();
-	case EStrideWarpingAxisMode::ActorSpaceVectorInput : 
-		{
-			const FVector WorldSpaceDir = MyAnimInstanceProxy->GetActorTransform().TransformVectorNoScale(UserSuppliedVector);
-			const FVector ComponentSpaceDir = MyAnimInstanceProxy->GetComponentTransform().InverseTransformVectorNoScale(WorldSpaceDir);
-			return ComponentSpaceDir.GetSafeNormal();
-		}
-	case EStrideWarpingAxisMode::WorldSpaceVectorInput: 
-		return MyAnimInstanceProxy->GetComponentTransform().InverseTransformVectorNoScale(UserSuppliedVector).GetSafeNormal();
-	}
-
-	return IKFootRootCSTransform.GetUnitAxis(EAxis::X);
-}
-
 void FAnimNode_StrideWarping::EvaluateSkeletalControl_AnyThread(FComponentSpacePoseContext& Output, TArray<FBoneTransform>& OutBoneTransforms)
 {
 	SCOPE_CYCLE_COUNTER(STAT_StrideWarping_Eval);
 	check(OutBoneTransforms.Num() == 0);
 
-	const FBoneContainer& RequiredBones = Output.Pose.GetPose().GetBoneContainer();
-
 	FTransform RootMotionTransformDelta;
 	const UE::Anim::IAnimRootMotionProvider* RootMotionProvider = UE::Anim::IAnimRootMotionProvider::Get();
 	bool bGraphDrivenWarping = RootMotionProvider && Mode == EWarpingEvaluationMode::Graph;
@@ -98,50 +64,56 @@ void FAnimNode_StrideWarping::EvaluateSkeletalControl_AnyThread(FComponentSpaceP
 	{
 		// Graph driven stride warping will override the manual stride direction with the intent of the current animation sub-graph's accumulated root motion
 		bGraphDrivenWarping = RootMotionProvider->ExtractRootMotion(Output.CustomAttributes, RootMotionTransformDelta);
-		if (bGraphDrivenWarping)
+		if (ensure(bGraphDrivenWarping))
 		{
-			ManualStrideWarpingDir = RootMotionTransformDelta.GetTranslation().GetSafeNormal2D();
+			StrideDirection = RootMotionTransformDelta.GetTranslation().GetSafeNormal();
 		}
 	}
 
-	FTransform IKFootRootTransform = Output.Pose.GetComponentSpaceTransform(IKFootRootBone.GetCompactPoseIndex(RequiredBones));
-	FVector StrideWarpingPlaneNormal = GetAxisModeValue(StrideWarpingAxisMode, IKFootRootTransform, ManualStrideWarpingDir);
-	const FVector FloorPlaneNormal = GetAxisModeValue(FloorNormalAxisMode, IKFootRootTransform, ManualFloorNormalInput);
-	const FVector GravityDir = GetAxisModeValue(GravityDirAxisMode, IKFootRootTransform, ManualGravityDirInput);
+	const FBoneContainer& RequiredBones = Output.Pose.GetPose().GetBoneContainer();
+	const FTransform IKFootRootTransform = Output.Pose.GetComponentSpaceTransform(IKFootRootBone.GetCompactPoseIndex(RequiredBones));
+	const FVector ResolvedFloorNormal = FloorNormalDirection.AsComponentSpaceDirection(AnimInstanceProxy, IKFootRootTransform);
+	const FVector ResolvedGravityDirection = GravityDirection.AsComponentSpaceDirection(AnimInstanceProxy, IKFootRootTransform);
 
-	if (bOrientStrideWarpingAxisBasedOnFloorNormal)
+	if (bOrientStrideDirectionUsingFloorNormal)
 	{
-		const FVector StrideWarpingAxisY = FloorPlaneNormal ^ StrideWarpingPlaneNormal;
-		StrideWarpingPlaneNormal = StrideWarpingAxisY ^ FloorPlaneNormal;
+		const FVector StrideWarpingAxis = ResolvedFloorNormal ^ StrideDirection;
+		StrideDirection = StrideWarpingAxis ^ ResolvedFloorNormal;
 	}
 
 #if ENABLE_ANIM_DEBUG
-	const bool bShowDebug = (CVarAnimNodeStrideWarpingDebug.GetValueOnAnyThread() == 1);
-	UWorld* DebugWorld = MyAnimInstanceProxy->GetSkelMeshComponent()->GetWorld();
-	if (bShowDebug)
+	bool bDebugging = false;
+#if WITH_EDITORONLY_DATA
+	bDebugging = bDebugging || bEnableDebugDraw;
+#endif
+	bDebugging = bDebugging || CVarAnimNodeStrideWarpingDebug.GetValueOnAnyThread() == 1;
+	if (bDebugging)
 	{
 		// Draw Floor Normal
-		MyAnimInstanceProxy->AnimDrawDebugDirectionalArrow(
-			MyAnimInstanceProxy->GetComponentTransform().TransformPosition(IKFootRootTransform.GetLocation()),
-			MyAnimInstanceProxy->GetComponentTransform().TransformPosition(IKFootRootTransform.GetLocation() + FloorPlaneNormal * 500.f),
-			50.f, FColor::Blue);
+		AnimInstanceProxy->AnimDrawDebugDirectionalArrow(
+			AnimInstanceProxy->GetComponentTransform().TransformPosition(IKFootRootTransform.GetLocation()),
+			AnimInstanceProxy->GetComponentTransform().TransformPosition(IKFootRootTransform.GetLocation() + ResolvedFloorNormal * 25.f),
+			40.f, FColor::Blue, false, 0.f, 2.f);
 	}
 #endif
 
 	// Get all foot IK Transforms
-	for (auto& FootData : FeetData)
+	for (auto& Foot : FootData)
 	{
-		FootData.IKBoneTransform = Output.Pose.GetComponentSpaceTransform(FootData.IKFootBoneIndex);
+		Foot.IKFootBoneTransform = Output.Pose.GetComponentSpaceTransform(Foot.IKFootBoneIndex);
 #if ENABLE_ANIM_DEBUG
-		if (bShowDebug)
+		if (bDebugging
+#if WITH_EDITORONLY_DATA
+			&& bDebugDrawIKFootOrigin
+#endif
+		)
 		{
-			FVector FootWorldLocation = MyAnimInstanceProxy->GetComponentTransform().TransformPosition(FootData.IKBoneTransform.GetLocation());
-			MyAnimInstanceProxy->AnimDrawDebugSphere(FootWorldLocation, 8.f, 16, FColor::Red);
+			const FVector FootWorldLocation = AnimInstanceProxy->GetComponentTransform().TransformPosition(Foot.IKFootBoneTransform.GetLocation());
+			AnimInstanceProxy->AnimDrawDebugSphere(FootWorldLocation, 8.f, 16, FColor::Red);
 		}
 #endif
 	}
 
-	float ComputedStrideScaling = ManualStrideScaling;
 	if (bGraphDrivenWarping)
 	{
 		// Early out when stride warping will have minimal pose contribution
@@ -152,196 +124,187 @@ void FAnimNode_StrideWarping::EvaluateSkeletalControl_AnyThread(FComponentSpaceP
 
 		const float RootMotionSpeed = RootMotionTransformDelta.GetTranslation().Size() / CachedDeltaTime;
 
-		if (FMath::IsNearlyZero(RootMotionSpeed, MinSpeedTolerance) || FMath::IsNearlyZero(LocomotionSpeed, MinSpeedTolerance))
+		if (FMath::IsNearlyZero(RootMotionSpeed, MinLocomotionSpeedThreshold) || FMath::IsNearlyZero(LocomotionSpeed, MinLocomotionSpeedThreshold))
 		{
-			ComputedStrideScaling = 1.f;
+			StrideScale = 1.f;
 		}
 		else
 		{
 			// Graph driven stride scale factor will be determined by the ratio of the
 			// locomotion (capsule/physics) speed against the animation root motion speed
-			ComputedStrideScaling = LocomotionSpeed / RootMotionSpeed;
+			StrideScale = LocomotionSpeed / RootMotionSpeed;
 		}
 	}
 
 	// Allow the opportunity for stride scale clamping and biasing regardless of evaluation mode
-	ComputedStrideScaling = StrideScalingScaleBiasClamp.ApplyTo(ComputedStrideScaling, CachedDeltaTime);
+	StrideScale = StrideScaleModifierState.ApplyTo(StrideScaleModifier, StrideScale, CachedDeltaTime);
 
 	if (bGraphDrivenWarping)
 	{
 		// Forward the side effects of stride warping on the root motion contribution for this sub-graph
-		RootMotionTransformDelta.ScaleTranslation(ComputedStrideScaling);
+		RootMotionTransformDelta.ScaleTranslation(StrideScale);
 		const bool bRootMotionOverridden = RootMotionProvider->OverrideRootMotion(RootMotionTransformDelta, Output.CustomAttributes);
 		ensure(bRootMotionOverridden);
 	}
 
-	// Scale IK feet bones along Stride Warping Axis, from the hip bone location.
-	for (auto& FootData : FeetData)
+	// Scale IK feet bones along Stride Warping Axis, from the Thigh bone location.
+	for (auto& Foot : FootData)
 	{
 		// Stride Warping along Stride Warping Axis
-		const FVector IKFootLocation = FootData.IKBoneTransform.GetLocation();
-		const FVector HipBoneLocation = Output.Pose.GetComponentSpaceTransform(FootData.HipBoneIndex).GetLocation();
+		const FVector IKFootLocation = Foot.IKFootBoneTransform.GetLocation();
+		const FVector ThighBoneLocation = Output.Pose.GetComponentSpaceTransform(Foot.ThighBoneIndex).GetLocation();
 
-		// Project Hip Bone Location on plane made of FootIKLocation and FloorPlaneNormal, along Gravity Dir.
+		// Project Thigh Bone Location on plane made of FootIKLocation and FloorPlaneNormal, along Gravity Dir.
 		// This will be the StrideWarpingPlaneOrigin
-		const FVector StrideWarpingPlaneOrigin = (FMath::Abs(GravityDir | FloorPlaneNormal) > DELTA) ? FMath::LinePlaneIntersection(HipBoneLocation, HipBoneLocation + GravityDir, IKFootLocation, FloorPlaneNormal) : IKFootLocation;
+		const FVector StrideWarpingPlaneOrigin = (FMath::Abs(ResolvedGravityDirection | ResolvedFloorNormal) > DELTA) ? FMath::LinePlaneIntersection(ThighBoneLocation, ThighBoneLocation + ResolvedGravityDirection, IKFootLocation, ResolvedFloorNormal) : IKFootLocation;
 
 		// Project FK Foot along StrideWarping Plane, this will be our Scale Origin
-		const FVector ScaleOrigin = FVector::PointPlaneProject(IKFootLocation, StrideWarpingPlaneOrigin, StrideWarpingPlaneNormal);
+		const FVector ScaleOrigin = FVector::PointPlaneProject(IKFootLocation, StrideWarpingPlaneOrigin, StrideDirection);
 
 		// Now the ScaleOrigin and IKFootLocation are forming a line parallel to the floor, and we can scale the IK foot.
-		const FVector WarpedLocation = ScaleOrigin + (IKFootLocation - ScaleOrigin) * ComputedStrideScaling;
-		FootData.IKBoneTransform.SetLocation(WarpedLocation);
+		const FVector WarpedLocation = ScaleOrigin + (IKFootLocation - ScaleOrigin) * StrideScale;
+		Foot.IKFootBoneTransform.SetLocation(WarpedLocation);
 
 #if ENABLE_ANIM_DEBUG
-		if (bShowDebug)
+		if (bDebugging
+#if WITH_EDITORONLY_DATA
+			&& bDebugDrawIKFootAdjustment
+#endif
+		)
 		{
-			FVector FootWorldLocation = MyAnimInstanceProxy->GetComponentTransform().TransformPosition(FootData.IKBoneTransform.GetLocation());
-			MyAnimInstanceProxy->AnimDrawDebugSphere(FootWorldLocation, 8.f, 16, FColor::Green);
+			const FVector FootWorldLocation = AnimInstanceProxy->GetComponentTransform().TransformPosition(Foot.IKFootBoneTransform.GetLocation());
+			AnimInstanceProxy->AnimDrawDebugSphere(FootWorldLocation, 8.f, 16, FColor::Green);
 
-			FVector ScaleOriginWorldLoc = MyAnimInstanceProxy->GetComponentTransform().TransformPosition(ScaleOrigin);
-			MyAnimInstanceProxy->AnimDrawDebugSphere(ScaleOriginWorldLoc, 8.f, 16, FColor::Yellow);
+			const FVector ScaleOriginWorldLoc = AnimInstanceProxy->GetComponentTransform().TransformPosition(ScaleOrigin);
+			AnimInstanceProxy->AnimDrawDebugSphere(ScaleOriginWorldLoc, 8.f, 16, FColor::Yellow);
 		}
 #endif
 	}
 
-	// Adjust Pelvis down if needed to keep foot contact with the ground and prevent over-extension
 	FVector PelvisOffset = FVector::ZeroVector;
 	const FCompactPoseBoneIndex PelvisBoneIndex = PelvisBone.GetCompactPoseIndex(RequiredBones);
+	check(PelvisBoneIndex != INDEX_NONE);
 
-	if ((PelvisBoneIndex != INDEX_NONE) && (PelvisPostAdjustmentAlpha > 0.f))
-	{
-		FTransform PelvisTransform = Output.Pose.GetComponentSpaceTransform(PelvisBoneIndex);
-		const FVector InitialPelvisLocation = PelvisTransform.GetLocation();
-
-		TArray<float> FKFeetDistToPelvis;
-		TArray<FVector> IKFeetLocation;
-	
-		for (auto& FootData : FeetData)
-		{
-			const FVector FKFootLocation = Output.Pose.GetComponentSpaceTransform(FootData.FKFootBoneIndex).GetLocation();
-			FKFeetDistToPelvis.Add(FVector::Dist(FKFootLocation, InitialPelvisLocation));
-
-			const FVector IKFootLocation = FootData.IKBoneTransform.GetLocation();
-			IKFeetLocation.Add(IKFootLocation);
-		}
+	FTransform PelvisTransform = Output.Pose.GetComponentSpaceTransform(PelvisBoneIndex);
+	const FVector InitialPelvisLocation = PelvisTransform.GetLocation();
 
-		// Pull Pelvis closer to feet iteratively
-		FVector AdjustedPelvisLocation = InitialPelvisLocation;
-		{
-			check(IKFeetLocation.Num() > 0);
-			const float PerFootWeight = (1.f / float(IKFeetLocation.Num()));
+	TArray<float, TInlineAllocator<10>> FKFootDistancesToPelvis;
+	FKFootDistancesToPelvis.Reserve(FootData.Num());
 
-			int32 Iterations = FMath::Clamp(PelvisAdjustmentMaxIter, 1, 10);
-			while (Iterations-- > 0)
-			{
-				const FVector PreAdjustmentLocation = AdjustedPelvisLocation;
-				AdjustedPelvisLocation = FVector::ZeroVector;
+	TArray<FVector, TInlineAllocator<10>> IKFootLocations;
+	IKFootLocations.Reserve(FootData.Num());
+	
+	for (auto& Foot : FootData)
+	{
+		const FVector FKFootLocation = Output.Pose.GetComponentSpaceTransform(Foot.FKFootBoneIndex).GetLocation();
+		FKFootDistancesToPelvis.Add(FVector::Dist(FKFootLocation, InitialPelvisLocation));
 
-				for (int32 Index = 0; Index < IKFeetLocation.Num(); Index++)
-				{
-					const FVector IdealPelvisLoc = IKFeetLocation[Index] + (PreAdjustmentLocation - IKFeetLocation[Index]).GetSafeNormal() * FKFeetDistToPelvis[Index];
-					AdjustedPelvisLocation += IdealPelvisLoc * PerFootWeight;
-				}
-			}
-		}
-		
-		// Apply spring between initial and adjusted spring location to smooth out change over time.
-		const FVector TargetAdjustment = (AdjustedPelvisLocation - InitialPelvisLocation);
-		PelvisAdjustmentInterp.Update(TargetAdjustment, CachedDeltaTime);
+		const FVector IKFootLocation = Foot.IKFootBoneTransform.GetLocation();
+		IKFootLocations.Add(IKFootLocation);
+	}
 
-		// Apply an alpha with the initial pelvis location, to retain some of the original motion. 
-		const FVector SmoothAdjustedPelvisLocation = InitialPelvisLocation + FMath::Lerp(FVector::ZeroVector, PelvisAdjustmentInterp.GetPosition(), PelvisPostAdjustmentAlpha);
-		PelvisTransform.SetLocation(SmoothAdjustedPelvisLocation);
+	// Adjust Pelvis down if needed to keep foot contact with the ground and prevent over-extension
+	PelvisTransform = PelvisIKFootSolver.Solve(PelvisTransform, FKFootDistancesToPelvis, IKFootLocations, CachedDeltaTime);
 
 #if ENABLE_ANIM_DEBUG
-		if (bShowDebug)
-		{
-			// Draw Adjustments in Pelvis location
-			MyAnimInstanceProxy->AnimDrawDebugSphere(MyAnimInstanceProxy->GetComponentTransform().TransformPosition(InitialPelvisLocation), 8.f, 16, FColor::Red);
-			MyAnimInstanceProxy->AnimDrawDebugSphere(MyAnimInstanceProxy->GetComponentTransform().TransformPosition(AdjustedPelvisLocation), 8.f, 16, FColor::Green);
-			MyAnimInstanceProxy->AnimDrawDebugSphere(MyAnimInstanceProxy->GetComponentTransform().TransformPosition(SmoothAdjustedPelvisLocation), 8.f, 16, FColor::Blue);
-
-			// Draw Stride Direction
-			MyAnimInstanceProxy->AnimDrawDebugDirectionalArrow(
-				MyAnimInstanceProxy->GetComponentTransform().TransformPosition(InitialPelvisLocation),
-				MyAnimInstanceProxy->GetComponentTransform().TransformPosition(InitialPelvisLocation + StrideWarpingPlaneNormal * 500.f),
-				50.f, FColor::Red);
-		}
+	if (bDebugging
+#if WITH_EDITORONLY_DATA
+		&& bDebugDrawPelvisAdjustment
 #endif
+	)
+	{
+		// Draw Adjustments in Pelvis location
+		AnimInstanceProxy->AnimDrawDebugSphere(AnimInstanceProxy->GetComponentTransform().TransformPosition(InitialPelvisLocation), 8.f, 16, FColor::Red);
+		AnimInstanceProxy->AnimDrawDebugSphere(AnimInstanceProxy->GetComponentTransform().TransformPosition(PelvisTransform.GetLocation()), 8.f, 16, FColor::Blue);
+	}
 
-		// Add adjusted pelvis transform
-		check(!PelvisTransform.ContainsNaN());
-		OutBoneTransforms.Add(FBoneTransform(PelvisBoneIndex, PelvisTransform));
-
-		// Compute final offset to use below
-		PelvisOffset = (PelvisTransform.GetLocation() - InitialPelvisLocation);
+	if (bDebugging)
+	{
+		// Draw Stride Direction
+		AnimInstanceProxy->AnimDrawDebugDirectionalArrow(
+			AnimInstanceProxy->GetComponentTransform().TransformPosition(InitialPelvisLocation),
+			AnimInstanceProxy->GetComponentTransform().TransformPosition(InitialPelvisLocation + StrideDirection * StrideScale * 100.f),
+			40.f, FColor::Red, false, 0.f, 2.f);
 	}
+#endif
+	// Add adjusted pelvis transform
+	check(!PelvisTransform.ContainsNaN());
+	OutBoneTransforms.Add(FBoneTransform(PelvisBoneIndex, PelvisTransform));
+
+	// Compute final offset to use below
+	PelvisOffset = (PelvisTransform.GetLocation() - InitialPelvisLocation);
 
 	// Rotate Thigh bones to help IK, and maintain leg shape.
-	if (bAdjustThighBonesRotation)
+	if (bCompensateIKUsingFKThighRotation)
 	{
-		for (auto& FootData : FeetData)
+		for (auto& Foot : FootData)
 		{
-			const FTransform HipTransform = Output.Pose.GetComponentSpaceTransform(FootData.HipBoneIndex);
-			const FTransform FKFootTransform = Output.Pose.GetComponentSpaceTransform(FootData.FKFootBoneIndex);
-			FTransform AdjustedHipTransform = HipTransform;
-
-			AdjustedHipTransform.AddToTranslation(PelvisOffset);
+			const FTransform ThighTransform = Output.Pose.GetComponentSpaceTransform(Foot.ThighBoneIndex);
+			const FTransform FKFootTransform = Output.Pose.GetComponentSpaceTransform(Foot.FKFootBoneIndex);
+			
+			FTransform AdjustedThighTransform = ThighTransform;
+			AdjustedThighTransform.AddToTranslation(PelvisOffset);
 
-			const FVector InitialDir = (FKFootTransform.GetLocation() - HipTransform.GetLocation()).GetSafeNormal();
-			const FVector TargetDir = (FootData.IKBoneTransform.GetLocation() - AdjustedHipTransform.GetLocation()).GetSafeNormal();
+			const FVector InitialDir = (FKFootTransform.GetLocation() - ThighTransform.GetLocation()).GetSafeNormal();
+			const FVector TargetDir = (Foot.IKFootBoneTransform.GetLocation() - AdjustedThighTransform.GetLocation()).GetSafeNormal();
 			
 #if ENABLE_ANIM_DEBUG
-			if (bShowDebug)
+			if (bDebugging
+#if WITH_EDITORONLY_DATA
+				&& bDebugDrawThighAdjustment
+#endif
+			)
 			{
-				MyAnimInstanceProxy->AnimDrawDebugLine(
-					MyAnimInstanceProxy->GetComponentTransform().TransformPosition(HipTransform.GetLocation()),
-					MyAnimInstanceProxy->GetComponentTransform().TransformPosition(FKFootTransform.GetLocation()),
-					FColor::Red);
-
-				MyAnimInstanceProxy->AnimDrawDebugLine(
-					MyAnimInstanceProxy->GetComponentTransform().TransformPosition(AdjustedHipTransform.GetLocation()),
-					MyAnimInstanceProxy->GetComponentTransform().TransformPosition(FootData.IKBoneTransform.GetLocation()),
-					FColor::Green);
+				AnimInstanceProxy->AnimDrawDebugLine(
+					AnimInstanceProxy->GetComponentTransform().TransformPosition(ThighTransform.GetLocation()),
+					AnimInstanceProxy->GetComponentTransform().TransformPosition(FKFootTransform.GetLocation()),
+					FColor::Red, false, 0.f, 2.f);
+
+				AnimInstanceProxy->AnimDrawDebugLine(
+					AnimInstanceProxy->GetComponentTransform().TransformPosition(AdjustedThighTransform.GetLocation()),
+					AnimInstanceProxy->GetComponentTransform().TransformPosition(Foot.IKFootBoneTransform.GetLocation()),
+					FColor::Green, false, 0.f, 2.f);
 			}
 #endif
-
 			// Find Delta Rotation take takes us from Old to New dir
 			const FQuat DeltaRotation = FQuat::FindBetweenNormals(InitialDir, TargetDir);
 
 			// Rotate our Joint quaternion by this delta rotation
-			AdjustedHipTransform.SetRotation(DeltaRotation * AdjustedHipTransform.GetRotation());
+			AdjustedThighTransform.SetRotation(DeltaRotation * AdjustedThighTransform.GetRotation());
 
-			// Add adjusted hip transform
-			check(!AdjustedHipTransform.ContainsNaN());
-			OutBoneTransforms.Add(FBoneTransform(FootData.HipBoneIndex, AdjustedHipTransform));
+			// Add adjusted thigh transform
+			check(!AdjustedThighTransform.ContainsNaN());
+			OutBoneTransforms.Add(FBoneTransform(Foot.ThighBoneIndex, AdjustedThighTransform));
 
 			// Clamp IK Feet bone based on FK leg. To prevent over-extension and preserve animated motion.
-			if (bClampIKUsingFKLeg)
+			if (bClampIKUsingFKLimits)
 			{
-				const float FKLength = FVector::Dist(FKFootTransform.GetLocation(), HipTransform.GetLocation());
-				const float IKLength = FVector::Dist(FootData.IKBoneTransform.GetLocation(), AdjustedHipTransform.GetLocation());
+				const float FKLength = FVector::Dist(FKFootTransform.GetLocation(), ThighTransform.GetLocation());
+				const float IKLength = FVector::Dist(Foot.IKFootBoneTransform.GetLocation(), AdjustedThighTransform.GetLocation());
 				if (IKLength > FKLength)
 				{
-					const FVector ClampedFootLocation = AdjustedHipTransform.GetLocation() + TargetDir * FKLength;
-					FootData.IKBoneTransform.SetLocation(ClampedFootLocation);
+					const FVector ClampedFootLocation = AdjustedThighTransform.GetLocation() + TargetDir * FKLength;
+					Foot.IKFootBoneTransform.SetLocation(ClampedFootLocation);
 				}
 			}
 		}
 	}
 
 	// Add final IK feet transforms
-	for (auto& FootData : FeetData)
+	for (auto& Foot : FootData)
 	{
 #if ENABLE_ANIM_DEBUG
-		if (bShowDebug)
+		if (bDebugging
+#if WITH_EDITORONLY_DATA
+			&& bDebugDrawIKFootFinal
+#endif
+		)
 		{
-			MyAnimInstanceProxy->AnimDrawDebugSphere(MyAnimInstanceProxy->GetComponentTransform().TransformPosition(FootData.IKBoneTransform.GetLocation()), 8.f, 16, FColor::Blue);
+			AnimInstanceProxy->AnimDrawDebugSphere(AnimInstanceProxy->GetComponentTransform().TransformPosition(Foot.IKFootBoneTransform.GetLocation()), 8.f, 16, FColor::Blue);
 		}
 #endif
-		check(!FootData.IKBoneTransform.ContainsNaN());
-		OutBoneTransforms.Add(FBoneTransform(FootData.IKFootBoneIndex, FootData.IKBoneTransform));
+		check(!Foot.IKFootBoneTransform.ContainsNaN());
+		OutBoneTransforms.Add(FBoneTransform(Foot.IKFootBoneIndex, Foot.IKFootBoneTransform));
 	}
 
 	// Sort OutBoneTransforms so indices are in increasing order.
@@ -351,55 +314,70 @@ void FAnimNode_StrideWarping::EvaluateSkeletalControl_AnyThread(FComponentSpaceP
 bool FAnimNode_StrideWarping::IsValidToEvaluate(const USkeleton* Skeleton, const FBoneContainer& RequiredBones)
 {
 #if ENABLE_ANIM_DEBUG
-	if(CVarAnimNodeStrideWarpingEnable.GetValueOnAnyThread() == 0)
+	if (CVarAnimNodeStrideWarpingEnable.GetValueOnAnyThread() == 0)
 	{
 		return false;
 	}
 #endif
-
-	return (FeetData.Num() > 0) 
-		&& (PelvisBone.GetCompactPoseIndex(RequiredBones) != INDEX_NONE) 
-		&& (IKFootRootBone.GetCompactPoseIndex(RequiredBones) != INDEX_NONE) 
-		&& (!FMath::IsNearlyEqual(StrideScalingScaleBiasClamp.ApplyTo(ManualStrideScaling, 0.f), 1.f, 0.001f) || PelvisAdjustmentInterp.IsInMotion() || Mode == EWarpingEvaluationMode::Graph);
-}
-
-FCompactPoseBoneIndex FAnimNode_StrideWarping::FindHipBoneIndex(const FCompactPoseBoneIndex& InFootBoneIndex, const int32& NumBonesInLimb, const FBoneContainer& RequiredBones) const
-{
-	FCompactPoseBoneIndex BoneIndex = InFootBoneIndex;
-	if (BoneIndex != INDEX_NONE)
+	if (PelvisBone.GetCompactPoseIndex(RequiredBones) == INDEX_NONE)
+	{
+		return false;
+	}
+	
+	if (IKFootRootBone.GetCompactPoseIndex(RequiredBones) == INDEX_NONE)
 	{
-		FCompactPoseBoneIndex ParentBoneIndex = RequiredBones.GetParentBoneIndex(BoneIndex);
+		return false;
+	}
 
-		int32 NumIterations = NumBonesInLimb;
-		while ((NumIterations-- > 0) && (ParentBoneIndex != INDEX_NONE))
+	if (FootData.IsEmpty())
+	{
+		return false;
+	}
+	else
+	{
+		for (const auto& Foot : FootData)
 		{
-			BoneIndex = ParentBoneIndex;
-			ParentBoneIndex = RequiredBones.GetParentBoneIndex(BoneIndex);
-		};
+			if (Foot.IKFootBoneIndex == INDEX_NONE || Foot.FKFootBoneIndex == INDEX_NONE || Foot.ThighBoneIndex == INDEX_NONE)
+			{
+				return false;
+			}
+		}
 	}
-
-	return BoneIndex;
+	
+	if (Mode == EWarpingEvaluationMode::Manual)
+	{
+		if (FMath::IsNearlyEqual(StrideScaleModifierState.ApplyTo(StrideScaleModifier, StrideScale, CachedDeltaTime), 1.f, KINDA_SMALL_NUMBER))
+		{
+			return false;
+		}
+	}
+	else if (LocomotionSpeed <= MinLocomotionSpeedThreshold)
+	{
+		return false;
+	}
+	return true;
 }
 
 void FAnimNode_StrideWarping::InitializeBoneReferences(const FBoneContainer& RequiredBones)
 {
 	IKFootRootBone.Initialize(RequiredBones);
 	PelvisBone.Initialize(RequiredBones);
+	FootData.Empty();
 
-	FeetData.Empty();
-	for (auto& FootDef : FeetDefinitions)
+	for (auto& Foot : FootDefinitions)
 	{
-		FootDef.IKFootBone.Initialize(RequiredBones);
-		FootDef.FKFootBone.Initialize(RequiredBones);
+		Foot.IKFootBone.Initialize(RequiredBones);
+		Foot.FKFootBone.Initialize(RequiredBones);
+		Foot.ThighBone.Initialize(RequiredBones);
 
-		FStrideWarpingFootData FootData;
-		FootData.IKFootBoneIndex = FootDef.IKFootBone.GetCompactPoseIndex(RequiredBones);
-		FootData.FKFootBoneIndex = FootDef.FKFootBone.GetCompactPoseIndex(RequiredBones);
-		FootData.HipBoneIndex = FindHipBoneIndex(FootData.FKFootBoneIndex, FMath::Max(FootDef.NumBonesInLimb, 1), RequiredBones);
+		FStrideWarpingFootData StrideFootData;
+		StrideFootData.IKFootBoneIndex = Foot.IKFootBone.GetCompactPoseIndex(RequiredBones);
+		StrideFootData.FKFootBoneIndex = Foot.FKFootBone.GetCompactPoseIndex(RequiredBones);
+		StrideFootData.ThighBoneIndex = Foot.ThighBone.GetCompactPoseIndex(RequiredBones);
 
-		if ((FootData.IKFootBoneIndex != INDEX_NONE) && (FootData.FKFootBoneIndex != INDEX_NONE) && (FootData.HipBoneIndex != INDEX_NONE))
+		if ((StrideFootData.IKFootBoneIndex != INDEX_NONE) && (StrideFootData.FKFootBoneIndex != INDEX_NONE) && (StrideFootData.ThighBoneIndex != INDEX_NONE))
 		{
-			FeetData.Add(FootData);
+			FootData.Add(StrideFootData);
 		}
 	}
 }
\ No newline at end of file
diff --git a/Engine/Plugins/Experimental/Animation/AnimationWarping/Source/Runtime/Public/BoneControllers/AnimNode_OrientationWarping.h b/Engine/Plugins/Experimental/Animation/AnimationWarping/Source/Runtime/Public/BoneControllers/AnimNode_OrientationWarping.h
--- a/Engine/Plugins/Experimental/Animation/AnimationWarping/Source/Runtime/Public/BoneControllers/AnimNode_OrientationWarping.h
+++ b/Engine/Plugins/Experimental/Animation/AnimationWarping/Source/Runtime/Public/BoneControllers/AnimNode_OrientationWarping.h
@@ -4,105 +4,70 @@
 
 #include "BoneControllers/BoneControllerTypes.h"
 #include "BoneControllers/AnimNode_SkeletalControlBase.h"
-#include "Animation/AnimNodeBase.h"
 #include "AnimNode_OrientationWarping.generated.h"
 
-USTRUCT()
-struct FOrientationWarpingSpineBoneSettings
-{
-	GENERATED_USTRUCT_BODY()
-
-	UPROPERTY(EditAnywhere, Category = "Settings")
-	FBoneReference Bone;
-};
-
-USTRUCT()
-struct FOrientationWarpingSpineBoneData
-{
-	GENERATED_USTRUCT_BODY()
-
-	FCompactPoseBoneIndex BoneIndex;
-	float Weight;
-
-	FOrientationWarpingSpineBoneData()
-		: BoneIndex(INDEX_NONE)
-		, Weight(0.f)
-	{}
-
-	FOrientationWarpingSpineBoneData(FCompactPoseBoneIndex InBoneIndex)
-		: BoneIndex(InBoneIndex)
-		, Weight(0.f)
-	{}
-
-	// Comparison Operator for Sorting.
-	struct FCompareBoneIndex
-	{
-		FORCEINLINE bool operator()(const FOrientationWarpingSpineBoneData& A, const FOrientationWarpingSpineBoneData& B) const
-		{
-			return A.BoneIndex < B.BoneIndex;
-		}
-	};
-};
-
-USTRUCT(BlueprintType)
-struct FOrientationWarpingSettings 
-{
-	GENERATED_USTRUCT_BODY()
-
-	/** Rotation Axis used to rotate mesh. */
-	UPROPERTY(EditAnywhere, Category = Settings)
-	TEnumAsByte<EAxis::Type> YawRotationAxis = EAxis::Z;
-
-	/** How much of rotation is on whole body versus on IK Feet. */
-	UPROPERTY(EditAnywhere, Category = "Settings")
-	float BodyOrientationAlpha = 0.5f;
-
-	/** Spine bones countering the rotation of the body, so character keeps aiming straight ahead. */
-	UPROPERTY(EditAnywhere, Category = "Settings")
-	TArray<FOrientationWarpingSpineBoneSettings> SpineBones;
-
-	/** IK Foot Root bone. To be rotated if 'BodyOrientationAlpha' is less than 1.f */
-	UPROPERTY(EditAnywhere, Category = "Settings")
-	FBoneReference IKFootRootBone;
-
-	/** IK Foot Bones. To be rotated if 'BodyOrientationAlpha' is less than 1.f */
-	UPROPERTY(EditAnywhere, Category = "Settings")
-	TArray<FBoneReference> IKFootBones;
-};
-
 USTRUCT(BlueprintInternalUseOnly)
 struct ANIMATIONWARPINGRUNTIME_API FAnimNode_OrientationWarping : public FAnimNode_SkeletalControlBase
 {
-	GENERATED_USTRUCT_BODY()
-
-	UPROPERTY(EditAnywhere, BlueprintReadWrite, Category = Evaluation, meta = (PinHiddenByDefault))
-	EWarpingEvaluationMode Mode;
-
-	UPROPERTY(EditAnywhere, BlueprintReadWrite, Category = Settings, meta = (PinShownByDefault))
-	float LocomotionAngle;
-
-	UPROPERTY(EditAnywhere, BlueprintReadWrite, Category = Settings, meta = (PinHiddenByDefault))
-	FOrientationWarpingSettings Settings;
-
-	UPROPERTY(Transient)
-	TArray<FOrientationWarpingSpineBoneData> SpineBoneDataArray;
+	GENERATED_BODY()
+
+	// Orientation warping evaluation mode (Graph or Manual)
+	UPROPERTY(EditAnywhere, Category=Evaluation)
+	EWarpingEvaluationMode Mode = EWarpingEvaluationMode::Manual;
+
+	// The desired orientation angle (in degrees) to warp by relative to the specified RotationAxis
+	UPROPERTY(EditAnywhere, BlueprintReadWrite, Category=Evaluation, meta=(PinShownByDefault))
+	float OrientationAngle = 0.f;
+
+	// The character locomotion angle (in degrees) relative to the specified RotationAxis
+	// This will be used in the following equation to compute the effective orientation warping angle: [Orientation = RotationBetween(RootMotionDirection, LocomotionDirection)]
+	UPROPERTY(EditAnywhere, BlueprintReadWrite, Category=Evaluation, meta=(PinShownByDefault))
+	float LocomotionAngle = 0.f;
+
+	// Specifies an angle threshold to prevent erroneous over-rotation of the character, disabled with a value of 0
+	//
+	// When the effective orientation warping angle is detected to be greater than this value (default: 90 degrees) the locomotion direction will be inverted prior to warping
+	// This will result in the following equation: [Orientation = RotationBetween(RootMotionDirection, -LocomotionDirection)]
+	//
+	// Example: Playing a forward running animation while the motion is going backward 
+	// Rather than orientation warping by 180 degrees, the system will warp by 0 degrees 
+	UPROPERTY(EditAnywhere, BlueprintReadWrite, Category=Evaluation, meta=(PinHiddenByDefault), meta=(ClampMin="0.0", ClampMax="180.0"))
+	float LocomotionAngleDeltaThreshold = 90.f;
+
+	// Spine bone definitions
+	// Used to counter the rotation of the body to keep the character facing towards the specified angle (Orientation Angle, Locomotion Angle)
+	// The amount of counter rotation applied is driven by DistributedBoneOrientationAlpha
+	UPROPERTY(EditAnywhere, Category=Settings)
+	TArray<FBoneReference> SpineBones;
+
+	// IK Foot Root Bone definition
+	UPROPERTY(EditAnywhere, Category=Settings, meta=(DisplayName="IK Foot Root Bone"))
+	FBoneReference IKFootRootBone;
 
-	TArray<FCompactPoseBoneIndex> IKFootBoneIndexArray;
+	// IK Foot definitions
+	UPROPERTY(EditAnywhere, Category=Settings, meta=(DisplayName="IK Foot Bones"))
+	TArray<FBoneReference> IKFootBones;
 
-	FCompactPoseBoneIndex IKFootRootBoneIndex;
+	// Specifies how much rotation is applied to the character body versus IK feet
+	UPROPERTY(EditAnywhere, Category=Settings, meta=(ClampMin="0.0", ClampMax="1.0"))
+	float DistributedBoneOrientationAlpha = 0.5f;
 
-	UPROPERTY(EditAnywhere, BlueprintReadWrite, Category = Settings, meta = (PinHiddenByDefault, EditCondition = "Mode == EWarpingEvaluationMode::Graph"))
-	float LocomotionInversionThresholdAngle = 90.f;
+	// Rotation axis used when rotating the character body
+	UPROPERTY(EditAnywhere, Category=Settings)
+	TEnumAsByte<EAxis::Type> RotationAxis = EAxis::Z;
 
-	UPROPERTY(EditAnywhere, BlueprintReadWrite, Category = Settings, meta = (PinHiddenByDefault))
-	float InterpSpeed = 0.f;
+	// Specifies the interpolation speed (in degrees per second) towards reaching the final warped rotation angle
+	// A value of 0 will cause instantaneous warping, while a greater value will introduce smoothing
+	UPROPERTY(EditAnywhere, Category=Settings, meta=(ClampMin="0.0"))
+	float RotationInterpSpeed = 10.f;
 
-	UPROPERTY(Transient)
-	float LastRotationAngle = 0.f;
+#if WITH_EDITORONLY_DATA
+	// Enable/Disable orientation warping debug drawing
+	UPROPERTY(EditAnywhere, Category=Debug)
+	bool bEnableDebugDraw = false;
+#endif
 
 public:
-	FAnimNode_OrientationWarping();
-
 	// FAnimNode_Base interface
 	virtual void GatherDebugData(FNodeDebugData& DebugData) override;
 	virtual void UpdateInternal(const FAnimationUpdateContext& Context) override;
@@ -114,9 +79,56 @@ public:
 	virtual bool IsValidToEvaluate(const USkeleton* Skeleton, const FBoneContainer& RequiredBones) override;
 	// End of FAnimNode_SkeletalControlBase interface
 
-
 private:
 	// FAnimNode_SkeletalControlBase interface
 	virtual void InitializeBoneReferences(const FBoneContainer& RequiredBones) override;
 	// End of FAnimNode_SkeletalControlBase interface
-};
+
+	struct FOrientationWarpingSpineBoneData
+	{
+		FCompactPoseBoneIndex BoneIndex;
+		float Weight;
+
+		FOrientationWarpingSpineBoneData()
+			: BoneIndex(INDEX_NONE)
+			, Weight(0.f)
+		{
+		}
+
+		FOrientationWarpingSpineBoneData(FCompactPoseBoneIndex InBoneIndex)
+			: BoneIndex(InBoneIndex)
+			, Weight(0.f)
+		{
+		}
+
+		// Comparison Operator for Sorting
+		struct FCompareBoneIndex
+		{
+			FORCEINLINE bool operator()(const FOrientationWarpingSpineBoneData& A, const FOrientationWarpingSpineBoneData& B) const
+			{
+				return A.BoneIndex < B.BoneIndex;
+			}
+		};
+	};
+
+	struct FOrientationWarpingFootData
+	{
+		TArray<FCompactPoseBoneIndex> IKFootBoneIndexArray;
+		FCompactPoseBoneIndex IKFootRootBoneIndex;
+
+		FOrientationWarpingFootData()
+			: IKFootBoneIndexArray()
+			, IKFootRootBoneIndex(INDEX_NONE)
+		{
+		}
+	};
+
+	// Computed spine bone indices and alpha weights for the specified spine definition 
+	TArray<FOrientationWarpingSpineBoneData> SpineBoneDataArray;
+
+	// Computed IK bone indices for the specified foot definitions 
+	FOrientationWarpingFootData IKFootData;
+
+	// Internally cached previous frame orientation warping angle
+	float PreviousWarpedRotation = 0.f;
+};
\ No newline at end of file
diff --git a/Engine/Plugins/Experimental/Animation/AnimationWarping/Source/Runtime/Public/BoneControllers/AnimNode_SlopeWarping.h b/Engine/Plugins/Experimental/Animation/AnimationWarping/Source/Runtime/Public/BoneControllers/AnimNode_SlopeWarping.h
--- a/Engine/Plugins/Experimental/Animation/AnimationWarping/Source/Runtime/Public/BoneControllers/AnimNode_SlopeWarping.h
+++ b/Engine/Plugins/Experimental/Animation/AnimationWarping/Source/Runtime/Public/BoneControllers/AnimNode_SlopeWarping.h
@@ -72,9 +72,6 @@ struct ANIMATIONWARPINGRUNTIME_API FAnimNode_SlopeWarping : public FAnimNode_Ske
 	UCharacterMovementComponent* MyMovementComponent;
 	FAnimInstanceProxy* MyAnimInstanceProxy;
 
-	UPROPERTY(VisibleDefaultsOnly, BlueprintReadOnly, Category = Evaluation, meta = (PinHiddenByDefault))
-	EWarpingEvaluationMode Mode;
-
 	/** IKFoot Root Bone. **/
 	UPROPERTY(EditAnywhere, Category = "Settings")
 	FBoneReference IKFootRootBone;
diff --git a/Engine/Plugins/Experimental/Animation/AnimationWarping/Source/Runtime/Public/BoneControllers/AnimNode_StrideWarping.h b/Engine/Plugins/Experimental/Animation/AnimationWarping/Source/Runtime/Public/BoneControllers/AnimNode_StrideWarping.h
--- a/Engine/Plugins/Experimental/Animation/AnimationWarping/Source/Runtime/Public/BoneControllers/AnimNode_StrideWarping.h
+++ b/Engine/Plugins/Experimental/Animation/AnimationWarping/Source/Runtime/Public/BoneControllers/AnimNode_StrideWarping.h
@@ -3,141 +3,128 @@
 #pragma once
 
 #include "BoneControllers/BoneControllerTypes.h"
+#include "BoneControllers/BoneControllerSolvers.h"
 #include "BoneControllers/AnimNode_SkeletalControlBase.h"
-#include "Engine/SpringInterpolator.h"
 #include "AnimNode_StrideWarping.generated.h"
 
-struct FInputScaleBiasClamp;
-
-UENUM(BlueprintType)
-enum class EStrideWarpingAxisMode : uint8
-{
-	IKFootRootLocalX,
-	IKFootRootLocalY,
-	IKFootRootLocalZ,
-	WorldSpaceVectorInput,
-	ComponentSpaceVectorInput,
-	ActorSpaceVectorInput,
-};
-
-/** Per foot definitions */
-USTRUCT()
-struct FStrideWarpingFootDefinition
+// Foot definition specifying the IK/FK foot bones and Thigh bone
+USTRUCT(BlueprintInternalUseOnly)
+struct ANIMATIONWARPINGRUNTIME_API FStrideWarpingFootDefinition
 {
-	GENERATED_USTRUCT_BODY()
+	GENERATED_BODY()
 
-	UPROPERTY(EditAnywhere, Category = "Settings")
+	// IK driven foot bone
+	UPROPERTY(EditAnywhere, Category=Settings, meta=(DisplayName="IK Foot Bone"))
 	FBoneReference IKFootBone;
 
-	UPROPERTY(EditAnywhere, Category = "Settings")
+	// FK driven foot bone
+	UPROPERTY(EditAnywhere, Category=Settings, meta=(DisplayName="FK Foot Bone"))
 	FBoneReference FKFootBone;
 
-	UPROPERTY(EditAnywhere, Category = "Settings")
-	int32 NumBonesInLimb;
-
-	FStrideWarpingFootDefinition()
-		: NumBonesInLimb(2)
-	{}
+	// Thigh bone, representing the end of the leg chain prior to reaching the Pelvis Bone 
+	UPROPERTY(EditAnywhere, Category=Settings)
+	FBoneReference ThighBone;
 };
 
-/** Runtime foot data after validation, we guarantee these bones to exist */
-USTRUCT()
-struct FStrideWarpingFootData
-{
-	GENERATED_USTRUCT_BODY()
-
-	FCompactPoseBoneIndex IKFootBoneIndex;
-	FCompactPoseBoneIndex FKFootBoneIndex;
-	FCompactPoseBoneIndex HipBoneIndex;
-	FTransform IKBoneTransform;
-
-	FStrideWarpingFootData()
-		: IKFootBoneIndex(INDEX_NONE)
-		, FKFootBoneIndex(INDEX_NONE)
-		, HipBoneIndex(INDEX_NONE)
-	{}
-};
-
-USTRUCT()
+USTRUCT(BlueprintInternalUseOnly)
 struct ANIMATIONWARPINGRUNTIME_API FAnimNode_StrideWarping : public FAnimNode_SkeletalControlBase
 {
-	GENERATED_USTRUCT_BODY()
-
-	FAnimNode_StrideWarping();
-
-	FAnimInstanceProxy* MyAnimInstanceProxy;
-
-	UPROPERTY(EditAnywhere, BlueprintReadWrite, Category = Evaluation, meta = (PinHiddenByDefault))
-	EWarpingEvaluationMode Mode;
-
-	UPROPERTY(EditAnywhere, BlueprintReadWrite, Category = Evaluation, meta = (PinHiddenByDefault))
-	FVector ManualStrideWarpingDir;
-
-	UPROPERTY(EditAnywhere, BlueprintReadWrite, Category = Evaluation, meta = (PinHiddenByDefault))
-	float ManualStrideScaling;
-
-	UPROPERTY(EditAnywhere, BlueprintReadWrite, Category = Evaluation, meta = (PinHiddenByDefault))
-	float LocomotionSpeed;
-
-	UPROPERTY(EditAnywhere, BlueprintReadWrite, Category = Evaluation, meta = (PinHiddenByDefault))
-	float MinSpeedTolerance;
+	GENERATED_BODY()
+
+	// Stride warping evaluation mode (Graph or Manual)
+	UPROPERTY(EditAnywhere, Category=Evaluation)
+	EWarpingEvaluationMode Mode = EWarpingEvaluationMode::Manual;
+
+	// Component-space stride direction
+	// Example: A value of <1,0,0> will warp the leg stride along the Forward Vector
+	UPROPERTY(EditAnywhere, BlueprintReadWrite, Category=Evaluation, meta=(PinShownByDefault))
+	FVector StrideDirection = FVector::ForwardVector;
+
+	// Stride scale, specifying the amount of warping applied to the IK foot definitions
+	// Example: A value of 0.5 will decrease the effective leg stride by half, while a value of 2.0 will double it
+	UPROPERTY(EditAnywhere, Category=Evaluation, meta=(ClampMin="0.0", PinShownByDefault))
+	float StrideScale = 1.f;
+
+	// Locomotion speed, specifying the current speed of the character
+	// This will be used in the following equation for computing the stride scale: [StrideScale = (LocomotionSpeed / RootMotionSpeed)]
+	// 
+	// Stride scale is a value specifying the amount of warping applied to the IK foot definitions
+	// Example: A value of 0.5 will decrease the effective leg stride by half, while a value of 2.0 will double it
+	UPROPERTY(EditAnywhere, BlueprintReadWrite, Category=Evaluation, meta=(ClampMin="0.0", PinShownByDefault))
+	float LocomotionSpeed = 0.f;
+
+	// Minimum locomotion speed required to apply stride warping
+	UPROPERTY(EditAnywhere, BlueprintReadWrite, Category=Evaluation, meta=(ClampMin="0.0", PinHiddenByDefault))
+	float MinLocomotionSpeedThreshold = 10.f;
+
+	// Pevlis Bone definition
+	UPROPERTY(EditAnywhere, Category=Settings)
+	FBoneReference PelvisBone;
 
-	UPROPERTY(EditAnywhere, Category = Settings)
+	// IK Foot Root Bone definition
+	UPROPERTY(EditAnywhere, Category=Settings, meta=(DisplayName="IK Foot Root Bone"))
 	FBoneReference IKFootRootBone;
 
-	UPROPERTY(EditAnywhere, Category = Settings)
-	TArray<FStrideWarpingFootDefinition> FeetDefinitions;
+	// Foot definitions specifying the IK, FK, and Thigh bone
+	UPROPERTY(EditAnywhere, Category=Settings)
+	TArray<FStrideWarpingFootDefinition> FootDefinitions;
 
-	UPROPERTY(Transient)
-	TArray<FStrideWarpingFootData> FeetData;
+	// Modifies the final stride scale value by optionally clamping and/or interpolating
+	UPROPERTY(EditAnywhere, Category=Settings)
+	FInputClampConstants StrideScaleModifier;
 
-	UPROPERTY(EditAnywhere, Category = Settings)
-	FBoneReference PelvisBone;
-
-	UPROPERTY(EditAnywhere, BlueprintReadWrite, Category = Settings)
-	EStrideWarpingAxisMode StrideWarpingAxisMode;
-
-	UPROPERTY(EditAnywhere, BlueprintReadWrite, Category = Settings)
-	EStrideWarpingAxisMode FloorNormalAxisMode;
+	// Floor normal direction, this value will internally convert into a corresponding Component-space representation prior to warping
+	// Default: World Space, Up Vector: <0,0,1>
+	UPROPERTY(EditAnywhere, BlueprintReadWrite, Category=Advanced, meta=(PinHiddenByDefault))
+	FWarpingVectorValue FloorNormalDirection = { EWarpingVectorMode::WorldSpaceVector, FVector::UpVector };
 
-	/** Direction vector for adjusting hips up and down. */
-	UPROPERTY(EditAnywhere, BlueprintReadWrite, Category = Settings)
-	EStrideWarpingAxisMode GravityDirAxisMode;
+	// Gravity direction, this value will internally convert into a corresponding Component-space representation prior to warping
+	// Default: World Space, Down Vector: <0,0,-1>
+	UPROPERTY(EditAnywhere, BlueprintReadWrite, Category=Advanced, meta=(PinHiddenByDefault))
+	FWarpingVectorValue GravityDirection = { EWarpingVectorMode::WorldSpaceVector, FVector::DownVector };
 
-	// Additional scaling, offsetting and clamping of PlayRate input.
-	// Performed after PlayRateBasis.
-	UPROPERTY(EditAnywhere, BlueprintReadWrite, Category = Settings)
-	FInputScaleBiasClamp StrideScalingScaleBiasClamp;
+	// Solver for controlling how much the pelvis is "pulled down" towards the IK/FK foot definitions during leg limb extension
+	UPROPERTY(EditAnywhere, Category=Advanced, meta=(DisplayName="Pelvis IK Foot Solver"))
+	FIKFootPelvisPullDownSolver PelvisIKFootSolver;
 
-	UPROPERTY(EditAnywhere, BlueprintReadWrite, Category = Settings, meta = (PinHiddenByDefault))
-	FVector ManualFloorNormalInput;
+	// Orients the specified (Manual) or computed (Graph) stride direction by the floor normal
+	UPROPERTY(EditAnywhere, Category=Advanced)
+	bool bOrientStrideDirectionUsingFloorNormal = true;
 
-	UPROPERTY(EditAnywhere, BlueprintReadWrite, Category = Settings, meta = (PinHiddenByDefault))
-	FVector ManualGravityDirInput;
+	// Include warping adjustment to the FK thigh bones alongside the IK/FK foot definitions
+	// This is used to help preserve the original overall leg shape
+	UPROPERTY(EditAnywhere, Category=Advanced, meta=(DisplayName="Compensate IK Using FK Thigh Rotation"))
+	bool bCompensateIKUsingFKThighRotation = true;
 
-	UPROPERTY(EditAnywhere, BlueprintReadWrite, Category = Settings)
-	float PelvisPostAdjustmentAlpha;
+	// Clamps the IK foot warping to prevent over-extension relative to the overall FK leg
+	UPROPERTY(EditAnywhere, Category=Advanced, meta=(DisplayName="Clamp IK Using FK Limits", EditCondition="bCompensateIKUsingFKThighRotation"))
+	bool bClampIKUsingFKLimits = true;
 
-	UPROPERTY(EditAnywhere, BlueprintReadWrite, Category = Settings)
-	int32 PelvisAdjustmentMaxIter;
+#if WITH_EDITORONLY_DATA
+	// Enable/Disable stride warping debug drawing
+	UPROPERTY(EditAnywhere, Category=Debug)
+	bool bEnableDebugDraw = false;
 
-	UPROPERTY(EditAnywhere, Category = Settings)
-	FVectorRK4SpringInterpolator PelvisAdjustmentInterp;
+	// Enable/Disable IK foot location debug drawing prior to warping
+	UPROPERTY(EditAnywhere, Category=Debug)
+	bool bDebugDrawIKFootOrigin = false;
 
-	UPROPERTY(EditAnywhere, Category = Settings)
-	uint32 bAdjustThighBonesRotation : 1;
+	// Enable/Disable IK foot location debug drawing following initial foot adjustment
+	UPROPERTY(EditAnywhere, Category=Debug)
+	bool bDebugDrawIKFootAdjustment = false;
 
-	UPROPERTY(EditAnywhere, Category = Settings)
-	uint32 bClampIKUsingFKLeg : 1;
+	// Enable/Disable pelvis debug drawing following adjustment
+	UPROPERTY(EditAnywhere, Category=Debug)
+	bool bDebugDrawPelvisAdjustment = false;
 
-	UPROPERTY(EditAnywhere, Category = Settings)
-	uint32 bOrientStrideWarpingAxisBasedOnFloorNormal : 1;
+	// Enable/Disable thigh debug drawing following adjustment
+	UPROPERTY(EditAnywhere, Category=Debug)
+	bool bDebugDrawThighAdjustment = false;
 
-	UPROPERTY(Transient)
-	float CachedDeltaTime;
-
-	FVector GetAxisModeValue(const EStrideWarpingAxisMode& AxisMode, const FTransform& IKFootRootCSTransform, const FVector& UserSuppliedVector) const;
-	FCompactPoseBoneIndex FindHipBoneIndex(const FCompactPoseBoneIndex& InFootBoneIndex, const int32& NumBonesInLimb, const FBoneContainer& RequiredBones) const;
+	// Enable/Disable IK foot location debug drawing following all adsjustments (Final warped result)
+	UPROPERTY(EditAnywhere, Category=Debug)
+	bool bDebugDrawIKFootFinal = false;
+#endif
 
 public:
 	// FAnimNode_Base interface
@@ -155,4 +142,32 @@ private:
 	// FAnimNode_SkeletalControlBase interface
 	virtual void InitializeBoneReferences(const FBoneContainer& RequiredBones) override;
 	// End of FAnimNode_SkeletalControlBase interface
+
+	struct FStrideWarpingFootData
+	{
+		FCompactPoseBoneIndex IKFootBoneIndex;
+		FCompactPoseBoneIndex FKFootBoneIndex;
+		FCompactPoseBoneIndex ThighBoneIndex;
+		FTransform IKFootBoneTransform;
+
+		FStrideWarpingFootData()
+			: IKFootBoneIndex(INDEX_NONE)
+			, FKFootBoneIndex(INDEX_NONE)
+			, ThighBoneIndex(INDEX_NONE)
+			, IKFootBoneTransform(FTransform::Identity)
+		{
+		}
+	};
+
+	// Internal cached anim instance proxy
+	FAnimInstanceProxy* AnimInstanceProxy = nullptr;
+
+	// Computed IK, FK, Thigh bone indices for the specified foot definitions
+	TArray<FStrideWarpingFootData> FootData;
+
+	// Internal cached delta time used for interpolators
+	float CachedDeltaTime = 0.f;
+
+	// Internal cached stride scale modifier state
+	FInputClampState StrideScaleModifierState;
 };
diff --git a/Engine/Source/Runtime/AnimGraphRuntime/Private/BoneControllers/BoneControllerSolvers.cpp b/Engine/Source/Runtime/AnimGraphRuntime/Private/BoneControllers/BoneControllerSolvers.cpp
new file mode 100644
--- /dev/null
+++ b/Engine/Source/Runtime/AnimGraphRuntime/Private/BoneControllers/BoneControllerSolvers.cpp
@@ -0,0 +1,61 @@
+// Copyright Epic Games, Inc. All Rights Reserved.
+
+#include "BoneControllers/BoneControllerSolvers.h"
+
+FTransform FIKFootPelvisPullDownSolver::Solve(FTransform PelvisTransform, TArrayView<const float> FKFootDistancesToPelvis, TArrayView<const FVector> IKFootLocations, float DeltaTime)
+{
+	const FVector InitialPelvisLocation = PelvisTransform.GetLocation();
+	
+	const int32 IKFootLocationsCount = IKFootLocations.Num();
+	const int32 FKFootDistancesCount = FKFootDistancesToPelvis.Num();
+	check(IKFootLocationsCount > 0);
+	check(FKFootDistancesCount > 0);
+	check(FKFootDistancesCount == IKFootLocationsCount);
+
+	FVector AdjustedPelvisLocation = InitialPelvisLocation;
+	FVector DeltaAdjustment = FVector::ZeroVector;
+
+	const float PerFootWeight = 1.f / static_cast<float>(IKFootLocationsCount);
+	const float AdjustmentDistMaxSquared = FMath::Pow(PelvisAdjustmentMaxDistance, 2.f);
+
+	// Pull pelvis closer to feet iteratively
+	for (int32 Iter = 0; Iter < PelvisAdjustmentMaxIter; ++Iter)
+	{
+		const FVector PreAdjustmentLocation = AdjustedPelvisLocation;
+		AdjustedPelvisLocation = FVector::ZeroVector;
+
+		// Apply pelvis adjustment contributions from all IK/FK foot chains
+		for (int32 Index = 0; Index < IKFootLocationsCount; ++Index)
+		{
+			const FVector IdealLocation = IKFootLocations[Index] + (PreAdjustmentLocation - IKFootLocations[Index]).GetSafeNormal() * FKFootDistancesToPelvis[Index];
+			AdjustedPelvisLocation += (IdealLocation * PerFootWeight);
+		}
+
+		const FVector PrevDeltaAdjustment = DeltaAdjustment;
+		DeltaAdjustment = AdjustedPelvisLocation - InitialPelvisLocation;
+		const float DeltaAdjustmentDist = FVector::Dist(PrevDeltaAdjustment, DeltaAdjustment);
+
+		// Keep track of how much delta adjustment is being applied per iteration
+		if (DeltaAdjustmentDist <= PelvisAdjustmentErrorTolerance)
+		{
+			break;
+		}
+	}
+
+	// Apply spring between initial and adjusted spring location to smooth out change over time
+	//PelvisAdjustmentInterp
+	PelvisAdjustmentInterp.Update(DeltaAdjustment, DeltaTime);
+
+	// Apply an alpha with the initial pelvis location, to retain some of the original motion
+	AdjustedPelvisLocation = InitialPelvisLocation + FMath::Lerp(FVector::ZeroVector, PelvisAdjustmentInterp.GetPosition(), PelvisAdjustmentInterpAlpha);
+
+	// Guarantee that we don't over-adjust the pelvis beyond the specified distance tolerance
+	if (FVector::DistSquared(AdjustedPelvisLocation, InitialPelvisLocation) >= AdjustmentDistMaxSquared)
+	{
+		DeltaAdjustment = AdjustedPelvisLocation - InitialPelvisLocation;
+		AdjustedPelvisLocation = InitialPelvisLocation + DeltaAdjustment.GetSafeNormal() * PelvisAdjustmentMaxDistance;
+	}
+	
+	PelvisTransform.SetLocation(AdjustedPelvisLocation);
+	return PelvisTransform;
+}
\ No newline at end of file
diff --git a/Engine/Source/Runtime/AnimGraphRuntime/Private/BoneControllers/BoneControllerTypes.cpp b/Engine/Source/Runtime/AnimGraphRuntime/Private/BoneControllers/BoneControllerTypes.cpp
new file mode 100644
--- /dev/null
+++ b/Engine/Source/Runtime/AnimGraphRuntime/Private/BoneControllers/BoneControllerTypes.cpp
@@ -0,0 +1,29 @@
+// Copyright Epic Games, Inc. All Rights Reserved.
+
+#include "BoneControllers/BoneControllerTypes.h"
+#include "Animation/AnimInstanceProxy.h"
+
+FVector FWarpingVectorValue::AsComponentSpaceDirection(const FAnimInstanceProxy* AnimInstanceProxy, const FTransform& IKFootRootTransform) const
+{
+	switch (Mode)
+	{
+	case EWarpingVectorMode::ComponentSpaceVector:
+		return Value.GetSafeNormal();
+	case EWarpingVectorMode::ActorSpaceVector:
+	{
+		const FVector WorldSpaceDirection = AnimInstanceProxy->GetActorTransform().TransformVectorNoScale(Value);
+		const FVector ComponentSpaceDirection = AnimInstanceProxy->GetComponentTransform().InverseTransformVectorNoScale(WorldSpaceDirection);
+		return ComponentSpaceDirection.GetSafeNormal();
+	}
+	case EWarpingVectorMode::WorldSpaceVector:
+		return AnimInstanceProxy->GetComponentTransform().InverseTransformVectorNoScale(Value).GetSafeNormal();
+	case EWarpingVectorMode::IKFootRootLocalSpaceVector:
+		return IKFootRootTransform.TransformVectorNoScale(Value.GetSafeNormal());
+	}
+#if WITH_EDITORONLY_DATA
+	const UEnum* TypeEnum = FindObject<UEnum>(ANY_PACKAGE, TEXT("EWarpingVectorMode"));
+	check(!!TypeEnum);
+	ensureMsgf(false, TEXT("AsComponentSpaceDirection specified an unhandled EWarpingVectorMode: %s"), *(TypeEnum->GetNameStringByIndex(static_cast<int32>(Mode))));
+#endif
+	return Value;
+}
\ No newline at end of file
diff --git a/Engine/Source/Runtime/AnimGraphRuntime/Public/BoneControllers/BoneControllerSolvers.h b/Engine/Source/Runtime/AnimGraphRuntime/Public/BoneControllers/BoneControllerSolvers.h
new file mode 100644
--- /dev/null
+++ b/Engine/Source/Runtime/AnimGraphRuntime/Public/BoneControllers/BoneControllerSolvers.h
@@ -0,0 +1,45 @@
+// Copyright Epic Games, Inc. All Rights Reserved.
+
+#pragma once
+
+#include "Containers/ArrayView.h"
+#include "Engine/SpringInterpolator.h"
+#include "UObject/ObjectMacros.h"
+#include "BoneControllerSolvers.generated.h"
+
+USTRUCT(BlueprintInternalUseOnly)
+struct ANIMGRAPHRUNTIME_API FIKFootPelvisPullDownSolver
+{
+	GENERATED_BODY()
+
+	// Specifies the spring interpolation parameters applied during pelvis adjustment
+	UPROPERTY(EditAnywhere, Category = Settings)
+	FVectorRK4SpringInterpolator PelvisAdjustmentInterp;
+
+	// Specifies an alpha between the original and final adjusted pelvis locations
+	// This is used to retain some degree of the original pelvis motion
+	UPROPERTY(EditAnywhere, Category=Settings, meta=(ClampMin="0.0", ClampMax="1.0"))
+	float PelvisAdjustmentInterpAlpha = 0.5f;
+
+	// Specifies the maximum displacement the pelvis can be adjusted relative to its original location
+	UPROPERTY(EditAnywhere, Category=Settings, meta=(ClampMin="0.0"))
+	float PelvisAdjustmentMaxDistance = 10.f;
+
+	// Specifies the pelvis adjustment distance error that is tolerated for each iteration of the solver
+	// 
+	// When it is detected that the pelvis adjustment distance is incrementing at a value lower or equal
+	// to this value for each iteration, the solve will halt. Lower values will marginally increase visual
+	// quality at the cost of performance, but may require a higher PelvisAdjustmentMaxIter to be specified
+	//
+	// The default value of 0.01 specifies 1 centimeter of error
+	UPROPERTY(EditAnywhere, Category=Advanced, meta=(ClampMin="0.001"))
+	float PelvisAdjustmentErrorTolerance = 0.01f;
+
+	// Specifies the maximum number of iterations to run for the pelvis adjustment solver
+	// Higher iterations will guarantee closer PelvisAdjustmentErrorTolerance convergence at the cost of performance
+	UPROPERTY(EditAnywhere, Category=Advanced, meta=(ClampMin="0"))
+	int32 PelvisAdjustmentMaxIter = 3;
+
+	// Iteratively pulls the character pelvis towards the ground based on the relationship of driven IK foot targets versus FK foot limits
+	FTransform Solve(FTransform PelvisTransform, TArrayView<const float> FKFootDistancesToPelvis, TArrayView<const FVector> IKFootLocations, float DeltaTime);
+};
\ No newline at end of file
diff --git a/Engine/Source/Runtime/AnimGraphRuntime/Public/BoneControllers/BoneControllerTypes.h b/Engine/Source/Runtime/AnimGraphRuntime/Public/BoneControllers/BoneControllerTypes.h
--- a/Engine/Source/Runtime/AnimGraphRuntime/Public/BoneControllers/BoneControllerTypes.h
+++ b/Engine/Source/Runtime/AnimGraphRuntime/Public/BoneControllers/BoneControllerTypes.h
@@ -2,18 +2,50 @@
 
 #pragma once
 
-#include "HAL/Platform.h"
 #include "UObject/ObjectMacros.h"
+#include "BoneContainer.h"
+#include "BoneControllerTypes.generated.h"
 
-UENUM(BlueprintType)
+struct FAnimInstanceProxy;
+
+UENUM(BlueprintInternalUseOnly)
 enum class EWarpingEvaluationMode : uint8
 {
-	/* Pose warping node parameters are entirely driven by the user */
+	// Pose warping node parameters are entirely driven by the user
 	Manual,
-	/**
-	* Pose warping nodes may participate in being graph-driven. This means some
-	* properties of the warp may be automatically configured by the accumulated root motion
-	* contribution of the animation sub-graph leading into the node
-	*/
+	// Pose warping nodes may participate in being graph-driven. This means some
+	// properties of the warp may be automatically configured by the accumulated 
+	// root motion delta contribution of the animation graph leading into the node
 	Graph
+};
+
+// The supported spaces of a corresponding input vector value
+UENUM(BlueprintInternalUseOnly)
+enum class EWarpingVectorMode : uint8
+{
+	// Component-space input vector
+	ComponentSpaceVector,
+	// Actor-space input vector
+	ActorSpaceVector,
+	// World-space input vector
+	WorldSpaceVector,
+	// IK Foot Root relative local-space input vector
+	IKFootRootLocalSpaceVector,
+};
+
+USTRUCT(BlueprintType)
+struct ANIMGRAPHRUNTIME_API FWarpingVectorValue
+{
+	GENERATED_BODY()
+
+	// Space of the corresponding Vector value
+	UPROPERTY(EditAnywhere, BlueprintReadWrite, Category=Settings, meta=(PinShownByDefault))
+	EWarpingVectorMode Mode = EWarpingVectorMode::ComponentSpaceVector;
+
+	// Specifies a vector relative to the space defined by Mode
+	UPROPERTY(EditAnywhere, BlueprintReadWrite, Category=Settings, meta=(PinShownByDefault))
+	FVector Value = FVector::ZeroVector;
+
+	// Retrieves a normalized Component-space direction from the specified DirectionMode and Direction value
+	FVector AsComponentSpaceDirection(const FAnimInstanceProxy* AnimInstanceProxy, const FTransform& IKFootRootTransform) const;
 };
\ No newline at end of file
diff --git a/Engine/Source/Runtime/Engine/Private/Animation/AnimNodeBase.cpp b/Engine/Source/Runtime/Engine/Private/Animation/AnimNodeBase.cpp
--- a/Engine/Source/Runtime/Engine/Private/Animation/AnimNodeBase.cpp
+++ b/Engine/Source/Runtime/Engine/Private/Animation/AnimNodeBase.cpp
@@ -492,7 +492,7 @@ void FComponentSpacePoseLink::EvaluateComponentSpace(FComponentSpacePoseContext&
 			{
 				Output.AnimInstanceProxy->RecordNodeAttribute(*Output.AnimInstanceProxy, SourceID, LinkID, UE::Anim::FAttributes::Attributes);
 
-				TArray<FName, FAnimStackAllocator> AttributeKeyNames;
+				TArray<FName, TInlineAllocator<8>> AttributeKeyNames;
 				if (Output.CustomAttributes.GetAllKeyNames(AttributeKeyNames))
 				{
 					for (const FName& AttributeKeyName : AttributeKeyNames)
diff --git a/Engine/Source/Runtime/Engine/Private/Animation/MotionTrajectoryTypes.cpp b/Engine/Source/Runtime/Engine/Private/Animation/MotionTrajectoryTypes.cpp
--- a/Engine/Source/Runtime/Engine/Private/Animation/MotionTrajectoryTypes.cpp
+++ b/Engine/Source/Runtime/Engine/Private/Animation/MotionTrajectoryTypes.cpp
@@ -6,13 +6,16 @@
 #include "Animation/AnimTypes.h"
 
 #if ENABLE_ANIM_DEBUG
-constexpr int32 DebugTrajectorySampleDisable = 0;
-constexpr int32 DebugTrajectorySampleCount = 1;
-constexpr int32 DebugTrajectorySampleTime = 2;
-constexpr int32 DebugTrajectorySampleDistance = 3;
-constexpr int32 DebugTrajectorySamplePosition = 4;
-constexpr int32 DebugTrajectorySampleVelocity = 5;
-constexpr int32 DebugTrajectorySampleAccel = 6;
+static constexpr int32 DebugTrajectorySampleDisable = 0;
+static constexpr int32 DebugTrajectorySampleCount = 1;
+static constexpr int32 DebugTrajectorySampleTime = 2;
+static constexpr int32 DebugTrajectorySampleDistance = 3;
+static constexpr int32 DebugTrajectorySamplePosition = 4;
+static constexpr int32 DebugTrajectorySampleVelocity = 5;
+static constexpr int32 DebugTrajectorySampleAccel = 6;
+static const FVector DebugSampleTypeOffset(0.f, 0.f, 50.f);
+static const FVector DebugSampleOffset(0.f, 0.f, 10.f);
+
 TAutoConsoleVariable<int32> CVarMotionTrajectoryDebug(TEXT("a.MotionTrajectory.Debug"), 0, TEXT("Turn on debug drawing for motion trajectory"));
 TAutoConsoleVariable<int32> CVarMotionTrajectoryDebugStride(TEXT("a.MotionTrajectory.Stride"), 1, TEXT("Configure the sample stride when displaying information"));
 TAutoConsoleVariable<int32> CVarMotionTrajectoryDebugOptions(TEXT("a.MotionTrajectory.Options"), 0,
@@ -235,27 +238,27 @@ void FTrajectorySampleRange::DebugDrawTrajectory(bool bEnable
 				FString DebugSampleString;
 				switch (DebugSampleOptions)
 				{
-				case 1: // Sample Index
+				case DebugTrajectorySampleCount: // Sample Index
 					DebugString = "Sample Index:";
 					DebugSampleString = DebugSampleString.Format(TEXT("{0}"), { Idx });
 ;					break;
-				case 2: // Sample Accumulated Time
+				case DebugTrajectorySampleTime: // Sample Accumulated Time
 					DebugString = "Sample Time:";
 					DebugSampleString = DebugSampleString.Format(TEXT("{0}"), { Samples[Idx].AccumulatedSeconds });
 					break;
-				case 3: // Sample Accumulated Distance
+				case DebugTrajectorySampleDistance: // Sample Accumulated Distance
 					DebugString = "Sample Distance:";
 					DebugSampleString = DebugSampleString.Format(TEXT("{0}"), { Samples[Idx].AccumulatedDistance });
 					break;
-				case 4: // Sample Position
+				case DebugTrajectorySamplePosition: // Sample Position
 					DebugString = "Sample Position:";
 					DebugSampleString = DebugSampleString.Format(TEXT("{0}"), { Samples[Idx].Transform.GetLocation().ToCompactString() });
 					break;
-				case 5: // Sample Velocity
+				case DebugTrajectorySampleVelocity: // Sample Velocity
 					DebugString = "Sample Velocity:";
 					DebugSampleString = DebugSampleString.Format(TEXT("{0}"), { Samples[Idx].LinearVelocity.ToCompactString() });
 					break;
-				case 6: // Sample Acceleration
+				case DebugTrajectorySampleAccel: // Sample Acceleration
 					DebugString = "Sample Acceleration:";
 					DebugSampleString = DebugSampleString.Format(TEXT("{0}"), { Samples[Idx].LinearAcceleration.ToCompactString() });
 					break;
@@ -269,10 +272,10 @@ void FTrajectorySampleRange::DebugDrawTrajectory(bool bEnable
 					// One time debug drawing of the per-sample type description
 					if (!DebugString.IsEmpty() && Idx == 0)
 					{
-						DrawDebugString(World, WorldTransform.GetLocation() + FVector(0.f, 0.f, 50.f), DebugString, nullptr, FColor::White, 0.f, false, 1.f);
+						DrawDebugString(World, WorldTransform.GetLocation() + DebugSampleTypeOffset, DebugString, nullptr, FColor::White, 0.f, false, 1.f);
 					}
 
-					DrawDebugString(World, WorldForward + FVector(0.f, 0.f, 10.f), DebugSampleString, nullptr, FColor::White, 0.f, false, 1.f);
+					DrawDebugString(World, WorldForward + DebugSampleOffset, DebugSampleString, nullptr, FColor::White, 0.f, false, 1.f);
 				}
 #endif
 			}
