Lyra post-mortem Pose Warping fixes and features for Stride Warping, Orientation Warping, and Sequence Evaluator:

Pose-Warping:
1. Compile-time errors for invalid bone definitions
2. Expanded SHOWDEBUG ANIMATION and ensure() logging messages for graph-driven mode
3. Interpolation state and controller fixes for both modes
4. Fixed and expanded in-game visualization

Orientation-Warping:
1. Fixed orientation angle calculation for manual mode
2. Overhauled in-game visualization

Slope Warping:
1. Marked node as Experimental for upcoming Animation Warping plugin migration from Experimental to Beta

Sequence Evaluator:
1. Removed erroneous delta time record calculation in update-phase, since this is already accounted for in the asset player tick

diff --git a/Engine/Plugins/Experimental/Animation/AnimationWarping/Source/Editor/Private/AnimGraph/AnimGraphNode_OrientationWarping.cpp b/Engine/Plugins/Experimental/Animation/AnimationWarping/Source/Editor/Private/AnimGraph/AnimGraphNode_OrientationWarping.cpp
--- a/Engine/Plugins/Experimental/Animation/AnimationWarping/Source/Editor/Private/AnimGraph/AnimGraphNode_OrientationWarping.cpp
+++ b/Engine/Plugins/Experimental/Animation/AnimationWarping/Source/Editor/Private/AnimGraph/AnimGraphNode_OrientationWarping.cpp
@@ -3,6 +3,7 @@
 #include "AnimGraph/AnimGraphNode_OrientationWarping.h"
 #include "Animation/AnimRootMotionProvider.h"
 #include "Kismet2/BlueprintEditorUtils.h"
+#include "Kismet2/CompilerResultsLog.h"
 #include "DetailCategoryBuilder.h"
 #include "DetailLayoutBuilder.h"
 #include "ScopedTransaction.h"
@@ -178,4 +179,87 @@ void UAnimGraphNode_OrientationWarping::GetOutputLinkAttributes(FNodeAttributeAr
 	}
 }
 
+void UAnimGraphNode_OrientationWarping::ValidateAnimNodeDuringCompilation(USkeleton* ForSkeleton, FCompilerResultsLog& MessageLog)
+{
+	auto HasInvalidBoneName = [](const FName& BoneName)
+	{
+		return BoneName == NAME_None;
+	};
+
+	auto HasInvalidBoneIndex = [&](const FName& BoneName)
+	{
+		return ForSkeleton->GetReferenceSkeleton().FindBoneIndex(BoneName) == INDEX_NONE;
+	};
+
+	auto InvalidBoneNameMessage = [&](const FName& BoneName)
+	{
+		FFormatNamedArguments Args;
+		Args.Add(TEXT("BoneName"), FText::FromName(BoneName));
+		const FText Message = FText::Format(NSLOCTEXT("OrientationWarping", "Invalid{BoneName}BoneName", "@@ - {BoneName} bone not found in Skeleton"), Args);
+		MessageLog.Warning(*Message.ToString(), this);
+	};
+
+	auto InvalidBoneIndexMessage = [&](const FName& BoneName)
+	{
+		FFormatNamedArguments Args;
+		Args.Add(TEXT("BoneName"), FText::FromName(BoneName));
+		const FText Message = FText::Format(NSLOCTEXT("OrientationWarping", "Invalid{BoneName}BoneInSkeleton", "@@ - {BoneName} bone definition is required"), Args);
+		MessageLog.Warning(*Message.ToString(), this);
+	};
+
+	if (Node.RotationAxis == EAxis::None)
+	{
+		MessageLog.Warning(*NSLOCTEXT("OrientationWarping", "InvalidRotationAxis", "@@ - Rotation Axis choice of X, Y, or Z is required").ToString(), this);
+	}
+
+	if (Node.SpineBones.IsEmpty())
+	{
+		MessageLog.Warning(*NSLOCTEXT("OrientationWarping", "InvalidSpineBones", "@@ - Spine bone definitions are required").ToString(), this);
+	}
+	else
+	{
+		for (const auto& Bone : Node.SpineBones)
+		{
+			if (HasInvalidBoneName(Bone.BoneName))
+			{
+				InvalidBoneIndexMessage("Spine");
+			}
+			else if (HasInvalidBoneIndex(Bone.BoneName))
+			{
+				InvalidBoneNameMessage(Bone.BoneName);
+			}
+		}
+	}
+
+	if (HasInvalidBoneName(Node.IKFootRootBone.BoneName))
+	{
+		InvalidBoneIndexMessage("IK Foot Root");
+	}
+	else if (HasInvalidBoneIndex(Node.IKFootRootBone.BoneName))
+	{
+		InvalidBoneNameMessage(Node.IKFootRootBone.BoneName);
+	}
+
+	if (Node.SpineBones.IsEmpty())
+	{
+		MessageLog.Warning(*NSLOCTEXT("OrientationWarping", "InvalidIKFootBones", "@@ - IK Foot bone definitions are required").ToString(), this);
+	}
+	else
+	{
+		for (const auto& Bone : Node.IKFootBones)
+		{
+			if (HasInvalidBoneName(Bone.BoneName))
+			{
+				InvalidBoneIndexMessage("IK Foot");
+			}
+			else if (HasInvalidBoneIndex(Bone.BoneName))
+			{
+				InvalidBoneNameMessage(Bone.BoneName);
+			}
+		}
+	}
+
+	Super::ValidateAnimNodeDuringCompilation(ForSkeleton, MessageLog);
+}
+
 #undef LOCTEXT_NAMESPACE
diff --git a/Engine/Plugins/Experimental/Animation/AnimationWarping/Source/Editor/Private/AnimGraph/AnimGraphNode_StrideWarping.cpp b/Engine/Plugins/Experimental/Animation/AnimationWarping/Source/Editor/Private/AnimGraph/AnimGraphNode_StrideWarping.cpp
--- a/Engine/Plugins/Experimental/Animation/AnimationWarping/Source/Editor/Private/AnimGraph/AnimGraphNode_StrideWarping.cpp
+++ b/Engine/Plugins/Experimental/Animation/AnimationWarping/Source/Editor/Private/AnimGraph/AnimGraphNode_StrideWarping.cpp
@@ -3,6 +3,7 @@
 #include "AnimGraph/AnimGraphNode_StrideWarping.h"
 #include "Animation/AnimRootMotionProvider.h"
 #include "Kismet2/BlueprintEditorUtils.h"
+#include "Kismet2/CompilerResultsLog.h"
 #include "DetailCategoryBuilder.h"
 #include "DetailLayoutBuilder.h"
 #include "ScopedTransaction.h"
@@ -210,4 +211,90 @@ void UAnimGraphNode_StrideWarping::GetOutputLinkAttributes(FNodeAttributeArray&
 	}
 }
 
+void UAnimGraphNode_StrideWarping::ValidateAnimNodeDuringCompilation(USkeleton* ForSkeleton, FCompilerResultsLog& MessageLog)
+{
+	auto HasInvalidBoneName = [](const FName& BoneName) 
+	{ 
+		return BoneName == NAME_None; 
+	};
+
+	auto HasInvalidBoneIndex = [&] (const FName& BoneName) 
+	{ 
+		return ForSkeleton->GetReferenceSkeleton().FindBoneIndex(BoneName) == INDEX_NONE; 
+	};
+
+	auto InvalidBoneNameMessage = [&](const FName& BoneName) 
+	{
+		FFormatNamedArguments Args;
+		Args.Add(TEXT("BoneName"), FText::FromName(BoneName));
+		const FText Message = FText::Format(NSLOCTEXT("StrideWarping", "Invalid{BoneName}BoneName", "@@ - {BoneName} bone not found in Skeleton"), Args);
+		MessageLog.Warning(*Message.ToString(), this);
+	};
+
+	auto InvalidBoneIndexMessage = [&](const FName& BoneName) 
+	{
+		FFormatNamedArguments Args;
+		Args.Add(TEXT("BoneName"), FText::FromName(BoneName));
+		const FText Message = FText::Format(NSLOCTEXT("StrideWarping", "Invalid{BoneName}BoneInSkeleton", "@@ - {BoneName} bone definition is required"), Args);
+		MessageLog.Warning(*Message.ToString(), this);
+	};
+
+	if (HasInvalidBoneName(Node.PelvisBone.BoneName))
+	{
+		InvalidBoneIndexMessage("Pelvis");
+	}
+	else if (HasInvalidBoneIndex(Node.PelvisBone.BoneName))
+	{
+		InvalidBoneNameMessage(Node.PelvisBone.BoneName);
+	}
+
+	if (HasInvalidBoneName(Node.IKFootRootBone.BoneName))
+	{
+		InvalidBoneIndexMessage("IK Foot Root");
+	}
+	else if (HasInvalidBoneIndex(Node.IKFootRootBone.BoneName))
+	{
+		InvalidBoneNameMessage(Node.IKFootRootBone.BoneName);
+	}
+
+	if (Node.FootDefinitions.IsEmpty())
+	{
+		MessageLog.Warning(*NSLOCTEXT("StrideWarping", "InvalidFootDefinitions", "@@ - Foot definitions are required").ToString(), this);
+	}
+	else
+	{
+		for (const auto& Foot : Node.FootDefinitions)
+		{
+			if (HasInvalidBoneName(Foot.IKFootBone.BoneName))
+			{
+				InvalidBoneIndexMessage("IK Foot");
+			}
+			else if (HasInvalidBoneIndex(Foot.IKFootBone.BoneName))
+			{
+				InvalidBoneNameMessage(Foot.IKFootBone.BoneName);
+			}
+
+			if (HasInvalidBoneName(Foot.FKFootBone.BoneName))
+			{
+				InvalidBoneIndexMessage("FK Foot");
+			}
+			else if (HasInvalidBoneIndex(Foot.FKFootBone.BoneName))
+			{
+				InvalidBoneNameMessage(Foot.FKFootBone.BoneName);
+			}
+
+			if (HasInvalidBoneName(Foot.ThighBone.BoneName))
+			{
+				InvalidBoneIndexMessage("Thigh");
+			}
+			else if (HasInvalidBoneIndex(Foot.ThighBone.BoneName))
+			{
+				InvalidBoneNameMessage(Foot.ThighBone.BoneName);
+			}
+		}
+	}
+
+	Super::ValidateAnimNodeDuringCompilation(ForSkeleton, MessageLog);
+}
+
 #undef LOCTEXT_NAMESPACE
diff --git a/Engine/Plugins/Experimental/Animation/AnimationWarping/Source/Editor/Public/AnimGraph/AnimGraphNode_OrientationWarping.h b/Engine/Plugins/Experimental/Animation/AnimationWarping/Source/Editor/Public/AnimGraph/AnimGraphNode_OrientationWarping.h
--- a/Engine/Plugins/Experimental/Animation/AnimationWarping/Source/Editor/Public/AnimGraph/AnimGraphNode_OrientationWarping.h
+++ b/Engine/Plugins/Experimental/Animation/AnimationWarping/Source/Editor/Public/AnimGraph/AnimGraphNode_OrientationWarping.h
@@ -22,6 +22,7 @@ public:
 	virtual void PostEditChangeProperty(struct FPropertyChangedEvent& PropertyChangedEvent) override;
 	virtual void GetInputLinkAttributes(FNodeAttributeArray& OutAttributes) const override;
 	virtual void GetOutputLinkAttributes(FNodeAttributeArray& OutAttributes) const override;
+	virtual void ValidateAnimNodeDuringCompilation(USkeleton* ForSkeleton, FCompilerResultsLog& MessageLog) override;
 	// End of UEdGraphNode interface
 
 protected:
diff --git a/Engine/Plugins/Experimental/Animation/AnimationWarping/Source/Editor/Public/AnimGraph/AnimGraphNode_SlopeWarping.h b/Engine/Plugins/Experimental/Animation/AnimationWarping/Source/Editor/Public/AnimGraph/AnimGraphNode_SlopeWarping.h
--- a/Engine/Plugins/Experimental/Animation/AnimationWarping/Source/Editor/Public/AnimGraph/AnimGraphNode_SlopeWarping.h
+++ b/Engine/Plugins/Experimental/Animation/AnimationWarping/Source/Editor/Public/AnimGraph/AnimGraphNode_SlopeWarping.h
@@ -7,7 +7,7 @@
 
 #include "AnimGraphNode_SlopeWarping.generated.h"
 
-UCLASS()
+UCLASS(Experimental)
 class ANIMATIONWARPINGEDITOR_API UAnimGraphNode_SlopeWarping : public UAnimGraphNode_SkeletalControlBase
 {
 	GENERATED_UCLASS_BODY()
diff --git a/Engine/Plugins/Experimental/Animation/AnimationWarping/Source/Editor/Public/AnimGraph/AnimGraphNode_StrideWarping.h b/Engine/Plugins/Experimental/Animation/AnimationWarping/Source/Editor/Public/AnimGraph/AnimGraphNode_StrideWarping.h
--- a/Engine/Plugins/Experimental/Animation/AnimationWarping/Source/Editor/Public/AnimGraph/AnimGraphNode_StrideWarping.h
+++ b/Engine/Plugins/Experimental/Animation/AnimationWarping/Source/Editor/Public/AnimGraph/AnimGraphNode_StrideWarping.h
@@ -24,6 +24,7 @@ public:
 	virtual void PostEditChangeProperty(struct FPropertyChangedEvent& PropertyChangedEvent) override;
 	virtual void GetInputLinkAttributes(FNodeAttributeArray& OutAttributes) const override;
 	virtual void GetOutputLinkAttributes(FNodeAttributeArray& OutAttributes) const override;
+	virtual void ValidateAnimNodeDuringCompilation(USkeleton* ForSkeleton, FCompilerResultsLog& MessageLog) override;
 	// End of UEdGraphNode interface
 
 protected:
diff --git a/Engine/Plugins/Experimental/Animation/AnimationWarping/Source/Runtime/Private/BoneControllers/AnimNode_OrientationWarping.cpp b/Engine/Plugins/Experimental/Animation/AnimationWarping/Source/Runtime/Private/BoneControllers/AnimNode_OrientationWarping.cpp
--- a/Engine/Plugins/Experimental/Animation/AnimationWarping/Source/Runtime/Private/BoneControllers/AnimNode_OrientationWarping.cpp
+++ b/Engine/Plugins/Experimental/Animation/AnimationWarping/Source/Runtime/Private/BoneControllers/AnimNode_OrientationWarping.cpp
@@ -7,9 +7,9 @@
 DECLARE_CYCLE_STAT(TEXT("OrientationWarping Eval"), STAT_OrientationWarping_Eval, STATGROUP_Anim);
 
 #if ENABLE_ANIM_DEBUG
-TAutoConsoleVariable<int32> CVarAnimNodeOrientationWarpingDebug(TEXT("a.AnimNode.OrientationWarping.Debug"), 0, TEXT("Turn on visualization debugging for Orientation Warping"));
-TAutoConsoleVariable<int32> CVarAnimNodeOrientationWarpingVerbose(TEXT("a.AnimNode.OrientationWarping.Verbose"), 0, TEXT("Turn on verbose graph debugging for Orientation Warping"));
-TAutoConsoleVariable<int32> CVarAnimNodeOrientationWarpingEnable(TEXT("a.AnimNode.OrientationWarping.Enable"), 1, TEXT("Toggle Orientation Warping"));
+static TAutoConsoleVariable<int32> CVarAnimNodeOrientationWarpingDebug(TEXT("a.AnimNode.OrientationWarping.Debug"), 0, TEXT("Turn on visualization debugging for Orientation Warping"));
+static TAutoConsoleVariable<int32> CVarAnimNodeOrientationWarpingVerbose(TEXT("a.AnimNode.OrientationWarping.Verbose"), 0, TEXT("Turn on verbose graph debugging for Orientation Warping"));
+static TAutoConsoleVariable<int32> CVarAnimNodeOrientationWarpingEnable(TEXT("a.AnimNode.OrientationWarping.Enable"), 1, TEXT("Toggle Orientation Warping"));
 #endif
 
 namespace UE::Anim
@@ -27,9 +27,10 @@ namespace UE::Anim
 		};
 	}
 
-	static inline bool IsInvalidWarpingAngle(float Angle, float Tolerance)
+	static inline bool IsInvalidWarpingAngleDegrees(float Angle, float Tolerance)
 	{
-		return FMath::IsNearlyZero(Angle, Tolerance) || FMath::IsNearlyEqual(FMath::Abs(Angle), PI, Tolerance);
+		Angle = FRotator::NormalizeAxis(Angle);
+		return FMath::IsNearlyZero(Angle, Tolerance) || FMath::IsNearlyEqual(FMath::Abs(Angle), 180.f, Tolerance);
 	}
 }
 
@@ -42,13 +43,13 @@ void FAnimNode_OrientationWarping::GatherDebugData(FNodeDebugData& DebugData)
 		if (Mode == EWarpingEvaluationMode::Manual)
 		{
 			DebugLine += TEXT("\n - Evaluation Mode: (Manual)");
-			DebugLine += FString::Printf(TEXT("\n - Orientation Angle: (%.3fd)"), FMath::RadiansToDegrees(OrientationAngle));
+			DebugLine += FString::Printf(TEXT("\n - Orientation Angle: (%.3fd)"), FMath::RadiansToDegrees(ActualOrientationAngle));
 		}
 		else
 		{
 			DebugLine += TEXT("\n - Evaluation Mode: (Graph)");
-			DebugLine += FString::Printf(TEXT("\n - Orientation Angle: (%.3fd)"), FMath::RadiansToDegrees(OrientationAngle));
-			DebugLine += FString::Printf(TEXT("\n - Locomotion Angle: (%.3fd)"), LocomotionAngle);
+			DebugLine += FString::Printf(TEXT("\n - Orientation Angle: (%.3fd)"), FMath::RadiansToDegrees(ActualOrientationAngle));
+			DebugLine += FString::Printf(TEXT("\n - Locomotion Angle: (%.3fd)"), FMath::RadiansToDegrees(LocomotionAngle));
 			DebugLine += FString::Printf(TEXT("\n - Locomotion Delta Angle Threshold: (%.3fd)"), LocomotionAngleDeltaThreshold);
 		}
 		DebugLine += FString::Printf(TEXT("\n - Distributed Bone Orientation Alpha: (%.3fd)"), DistributedBoneOrientationAlpha);
@@ -61,7 +62,7 @@ void FAnimNode_OrientationWarping::GatherDebugData(FNodeDebugData& DebugData)
 	else
 #endif
 	{
-		DebugLine += FString::Printf(TEXT("(Orientation Angle: %.3fd)"), FMath::RadiansToDegrees(OrientationAngle));
+	DebugLine += FString::Printf(TEXT("(Orientation Angle: %.3fd)"), FMath::RadiansToDegrees(ActualOrientationAngle));
 	}
 	DebugData.AddDebugItem(DebugLine);
 	ComponentPose.GatherDebugData(DebugData);
@@ -70,6 +71,10 @@ void FAnimNode_OrientationWarping::GatherDebugData(FNodeDebugData& DebugData)
 void FAnimNode_OrientationWarping::Initialize_AnyThread(const FAnimationInitializeContext& Context)
 {
 	FAnimNode_SkeletalControlBase::Initialize_AnyThread(Context);
+
+	PreviousRootMotionDeltaDirection = FVector::ZeroVector;
+	PreviousOrientationAngle = 0.f;
+	ActualOrientationAngle = 0.f;
 }
 
 void FAnimNode_OrientationWarping::UpdateInternal(const FAnimationUpdateContext& Context)
@@ -82,8 +87,20 @@ void FAnimNode_OrientationWarping::EvaluateSkeletalControl_AnyThread(FComponentS
 	SCOPE_CYCLE_COUNTER(STAT_OrientationWarping_Eval);
 	check(OutBoneTransforms.Num() == 0);
 
+	ActualOrientationAngle = OrientationAngle;
+
+	const FVector RotationAxisVector = UE::Anim::GetAxisVector(RotationAxis);
+	FVector RootMotionDeltaDirection = FVector::ZeroVector;
+	FVector LocomotionForward = FVector::ZeroVector;
+
+	bool bGraphDrivenWarping = false;
 	const UE::Anim::IAnimRootMotionProvider* RootMotionProvider = UE::Anim::IAnimRootMotionProvider::Get();
-	const bool bGraphDrivenWarping = RootMotionProvider && Mode == EWarpingEvaluationMode::Graph;
+
+	if (Mode == EWarpingEvaluationMode::Graph)
+	{
+		bGraphDrivenWarping = !!RootMotionProvider;
+		ensureMsgf(bGraphDrivenWarping, TEXT("Graph driven Orientation Warping expected a valid root motion delta provider interface."));
+	}
 
 	// We will likely need to revisit LocomotionAngle participating as an input to orientation warping.
 	// Without velocity information from the motion model (such as the capsule), LocomotionAngle isn't enough
@@ -98,16 +115,16 @@ void FAnimNode_OrientationWarping::EvaluateSkeletalControl_AnyThread(FComponentS
 	// The solution may be instead to pass velocity with the actor base rotation, allowing us to retain
 	// speed information about the motion. It may also allow us to do more complex orienting behavior 
 	// when multiple degrees of freedom can be considered.
-	OrientationAngle = FMath::DegreesToRadians(FRotator::NormalizeAxis(bGraphDrivenWarping ? LocomotionAngle : OrientationAngle));
-	const FVector RotationAxisVector = UE::Anim::GetAxisVector(RotationAxis);
 
-	// Graph driven orientation warping will modify the incoming root motion to orient towards the intended locomotion angle
 	if (bGraphDrivenWarping)
 	{
-		FTransform RootMotionTransformDelta;
-		const bool bRootMotionDeltaPresent = RootMotionProvider->ExtractRootMotion(Output.CustomAttributes, RootMotionTransformDelta);
+		FTransform RootMotionTransformDelta = FTransform::Identity;
+		bGraphDrivenWarping = RootMotionProvider->ExtractRootMotion(Output.CustomAttributes, RootMotionTransformDelta);
 
-		if (ensure(bRootMotionDeltaPresent))
+		// Graph driven orientation warping will modify the incoming root motion to orient towards the intended locomotion angle
+		if (ensureMsgf(bGraphDrivenWarping, TEXT(
+			"Graph driven Orientation Warping expected a root motion delta to be present in the attribute stream.\n"
+			"Check that the evaluating Animation Sequence(s) have: \"Enable Root Motion\" enabled.")))
 		{
 			// In UE, forward is defined as +x; consequently this is also true when sampling an actor's velocity. Historically the skeletal 
 			// mesh component forward will not match the actor, requiring us to correct the rotation before sampling the LocomotionForward.
@@ -117,77 +134,122 @@ void FAnimNode_OrientationWarping::EvaluateSkeletalControl_AnyThread(FComponentS
 			// 1. Actor Forward
 			// 2. Actor Velocity
 			// 3. Skeletal Mesh Relative Rotation
+
+			LocomotionAngle = FRotator::NormalizeAxis(LocomotionAngle);
+			LocomotionAngle = FMath::DegreesToRadians(LocomotionAngle);
+			const FQuat LocomotionRotation = FQuat(RotationAxisVector, LocomotionAngle);
+
 			const FTransform SkeletalMeshRelativeTransform = Output.AnimInstanceProxy->GetComponentRelativeTransform();
 			const FQuat SkeletalMeshRelativeRotation = SkeletalMeshRelativeTransform.GetRotation();
-			const FQuat LocomotionRotation = FQuat(RotationAxisVector, OrientationAngle);
-			const FVector LocomotionForwardDir = SkeletalMeshRelativeRotation.UnrotateVector(LocomotionRotation.GetForwardVector());
-			const FVector RootMotionDeltaDir = RootMotionTransformDelta.GetTranslation();
+			LocomotionForward = SkeletalMeshRelativeRotation.UnrotateVector(LocomotionRotation.GetForwardVector()).GetSafeNormal();
+
+			const FVector RootMotionDeltaTranslation = RootMotionTransformDelta.GetTranslation();
+			RootMotionDeltaDirection = RootMotionDeltaTranslation.GetSafeNormal();
 
 			// Capture the delta rotation from the axis of motion we care about
-			FQuat WarpedRotation = FQuat::FindBetween(RootMotionDeltaDir, LocomotionForwardDir);
-			OrientationAngle = WarpedRotation.GetTwistAngle(RotationAxisVector);
+			FQuat WarpedRotation = FQuat::FindBetween(RootMotionDeltaDirection, LocomotionForward);
+			ActualOrientationAngle = WarpedRotation.GetTwistAngle(RotationAxisVector);
 
 			// Motion Matching may return an animation that deviates a lot from the movement direction (e.g movement direction going bwd and motion matching could return the fwd animation for a few frames)
 			// When that happens, since we use the delta between root motion and movement direction, we would be over-rotating the lower body and breaking the pose during those frames
 			// So, when that happens we use the inverse of the movement direction to calculate our target rotation. 
 			// This feels a bit 'hacky' but its the only option I've found so far to mitigate the problem
-			if (LocomotionAngleDeltaThreshold > 0.f && FMath::Abs(FMath::RadiansToDegrees(OrientationAngle)) > LocomotionAngleDeltaThreshold)
+			if (LocomotionAngleDeltaThreshold > 0.f && FMath::Abs(FMath::RadiansToDegrees(ActualOrientationAngle)) > LocomotionAngleDeltaThreshold)
 			{
-				WarpedRotation = FQuat::FindBetween(RootMotionDeltaDir, -LocomotionForwardDir);
-				OrientationAngle = WarpedRotation.GetTwistAngle(RotationAxisVector);
+				WarpedRotation = FQuat::FindBetween(RootMotionDeltaDirection, -LocomotionForward);
+				ActualOrientationAngle = WarpedRotation.GetTwistAngle(RotationAxisVector);
 			}
 
+			// For interpolated warping, guarantee that PreviousOrientationAngle is respect to the current frame's root motion direction 
+			const FVector RootMotionCrossProduct = RootMotionDeltaDirection.Cross(PreviousRootMotionDeltaDirection);
+			float RootMotionDeltaAngle = FMath::Acos(RootMotionDeltaDirection.Dot(PreviousRootMotionDeltaDirection));
+			RootMotionDeltaAngle *= FMath::Sign(RotationAxisVector.Dot(RootMotionCrossProduct));
+
+			PreviousRootMotionDeltaDirection = RootMotionDeltaDirection;
+			PreviousOrientationAngle += RootMotionDeltaAngle;
+
 			// Rotate the root motion delta fully by the warped angle
-			const FVector RootMotionTranslationDelta = RootMotionTransformDelta.GetTranslation();
-			const FVector WarpedRootMotionTranslationDelta = WarpedRotation.RotateVector(RootMotionTranslationDelta);
+			const FVector WarpedRootMotionTranslationDelta = WarpedRotation.RotateVector(RootMotionDeltaTranslation);
 			RootMotionTransformDelta.SetTranslation(WarpedRootMotionTranslationDelta);
 
 			// Forward the side effects of orientation warping on the root motion contribution for this sub-graph
 			const bool bRootMotionOverridden = RootMotionProvider->OverrideRootMotion(RootMotionTransformDelta, Output.CustomAttributes);
-			ensure(bRootMotionOverridden);
-#if ENABLE_ANIM_DEBUG
-			bool bDebugging = false;
-#if WITH_EDITORONLY_DATA
-			bDebugging = bDebugging || bEnableDebugDraw;
-#endif
-			bDebugging = bDebugging || CVarAnimNodeOrientationWarpingDebug.GetValueOnAnyThread() == 1;
-
-			if (bDebugging)
-			{
-				const FVector LocomotionDir = Output.AnimInstanceProxy->GetActorTransform().TransformVectorNoScale(LocomotionRotation.GetForwardVector()).GetSafeNormal();
-				const FVector RootMotionDir = Output.AnimInstanceProxy->GetActorTransform().TransformVectorNoScale(SkeletalMeshRelativeRotation.RotateVector(RootMotionDeltaDir)).GetSafeNormal();
-
-				Output.AnimInstanceProxy->AnimDrawDebugDirectionalArrow(
-					Output.AnimInstanceProxy->GetComponentTransform().GetLocation(),
-					Output.AnimInstanceProxy->GetComponentTransform().GetLocation() + LocomotionDir * 100.f,
-					40.f, FColor::Red, false, 0.f, 2.f);
-
-				Output.AnimInstanceProxy->AnimDrawDebugDirectionalArrow(
-					Output.AnimInstanceProxy->GetComponentTransform().GetLocation(),
-					Output.AnimInstanceProxy->GetComponentTransform().GetLocation() + RootMotionDir * 100.f,
-					40.f, FColor::Blue, false, 0.f, 2.f);
-			}
-#endif
+			ensureMsgf(bRootMotionOverridden, TEXT("Graph driven Orientation Warping expected a root motion delta to be present in the attribute stream prior to warping/overriding it."));
+		}
+		else
+		{
+			// Early exit on broken graph driven behavior 
+			return;
 		}
 	} 
-	else if (UE::Anim::IsInvalidWarpingAngle(OrientationAngle, KINDA_SMALL_NUMBER))
+	else
 	{
-		return;
+		// Manual orientation warping will take the angle directly
+		ActualOrientationAngle = FRotator::NormalizeAxis(ActualOrientationAngle);
+		ActualOrientationAngle = FMath::DegreesToRadians(ActualOrientationAngle);
 	}
 
+	// Optionally interpolate the effective orientation towards the target orientation angle
 	if (RotationInterpSpeed > 0.f)
 	{
-		OrientationAngle = FMath::FInterpTo(PreviousWarpedRotation, OrientationAngle, Output.AnimInstanceProxy->GetDeltaSeconds(), RotationInterpSpeed);
-		PreviousWarpedRotation = OrientationAngle;
+		ActualOrientationAngle = FMath::FInterpTo(PreviousOrientationAngle, ActualOrientationAngle, Output.AnimInstanceProxy->GetDeltaSeconds(), RotationInterpSpeed);
+		PreviousOrientationAngle = ActualOrientationAngle;
 	}
 
 	// Allow the alpha value of the node to affect the final rotation
-	OrientationAngle *= ActualAlpha;
+	ActualOrientationAngle *= ActualAlpha;
 
+#if ENABLE_ANIM_DEBUG
+	bool bDebugging = false;
+#if WITH_EDITORONLY_DATA
+	bDebugging = bDebugging || bEnableDebugDraw;
+#else
+	constexpr float DebugDrawScale = 1.f;
+#endif
+	bDebugging = bDebugging || CVarAnimNodeOrientationWarpingDebug.GetValueOnAnyThread() == 1;
+
+	if (bDebugging)
+	{
+		const FTransform ComponentTransform = Output.AnimInstanceProxy->GetComponentTransform();
+		const FVector ActorForwardDirection = Output.AnimInstanceProxy->GetActorTransform().GetRotation().GetForwardVector();
+
+		const FVector ForwardDirection = bGraphDrivenWarping
+			? ComponentTransform.GetRotation().RotateVector(LocomotionForward)
+			: ActorForwardDirection;
+
+		FVector DebugArrowOffset = FVector::ZAxisVector * DebugDrawScale;
+		Output.AnimInstanceProxy->AnimDrawDebugDirectionalArrow(
+			ComponentTransform.GetLocation() + DebugArrowOffset,
+			ComponentTransform.GetLocation() + DebugArrowOffset + ForwardDirection * 100.f * DebugDrawScale,
+			40.f * DebugDrawScale, FColor::Red, false, 0.f, 2.f * DebugDrawScale);
+
+		const FVector RotationDirection = bGraphDrivenWarping
+			? ComponentTransform.GetRotation().RotateVector(RootMotionDeltaDirection)
+			: ActorForwardDirection.RotateAngleAxis(OrientationAngle, RotationAxisVector);
+
+		DebugArrowOffset += DebugArrowOffset;
+		Output.AnimInstanceProxy->AnimDrawDebugDirectionalArrow(
+			ComponentTransform.GetLocation() + DebugArrowOffset,
+			ComponentTransform.GetLocation() + DebugArrowOffset + RotationDirection * 100.f * DebugDrawScale,
+			40.f * DebugDrawScale, FColor::Blue, false, 0.f, 2.f * DebugDrawScale);
+
+		const float ActualOrientationAngleDegrees = FMath::RadiansToDegrees(ActualOrientationAngle);
+		const FVector WarpedRotationDirection = bGraphDrivenWarping 
+			? RotationDirection.RotateAngleAxis(ActualOrientationAngleDegrees, RotationAxisVector)
+			: ActorForwardDirection.RotateAngleAxis(ActualOrientationAngleDegrees, RotationAxisVector);
+
+		DebugArrowOffset += DebugArrowOffset;
+		Output.AnimInstanceProxy->AnimDrawDebugDirectionalArrow(
+			ComponentTransform.GetLocation() + DebugArrowOffset,
+			ComponentTransform.GetLocation() + DebugArrowOffset + WarpedRotationDirection * 100.f * DebugDrawScale,
+			40.f * DebugDrawScale, FColor::Green, false, 0.f, 2.f * DebugDrawScale);
+	}
+#endif
+		
 	// Rotate Root Bone first, as that cheaply rotates the whole pose with one transformation.
 	if (!FMath::IsNearlyZero(DistributedBoneOrientationAlpha, KINDA_SMALL_NUMBER))
 	{
-		const FQuat RootRotation = FQuat(RotationAxisVector, OrientationAngle * DistributedBoneOrientationAlpha);
+		const FQuat RootRotation = FQuat(RotationAxisVector, ActualOrientationAngle * DistributedBoneOrientationAlpha);
 		const FCompactPoseBoneIndex RootBoneIndex(0);
 
 		FTransform RootBoneTransform(Output.Pose.GetComponentSpaceTransform(RootBoneIndex));
@@ -195,7 +257,7 @@ void FAnimNode_OrientationWarping::EvaluateSkeletalControl_AnyThread(FComponentS
 		RootBoneTransform.NormalizeRotation();
 		Output.Pose.SetComponentSpaceTransform(RootBoneIndex, RootBoneTransform);
 	}
-		
+
 	const int32 NumSpineBones = SpineBoneDataArray.Num();
 	const bool bSpineOrientationAlpha = !FMath::IsNearlyZero(DistributedBoneOrientationAlpha, KINDA_SMALL_NUMBER);
 	const bool bUpdateSpineBones = (NumSpineBones > 0) && bSpineOrientationAlpha;
@@ -206,7 +268,7 @@ void FAnimNode_OrientationWarping::EvaluateSkeletalControl_AnyThread(FComponentS
 		for (int32 ArrayIndex = 0; ArrayIndex < NumSpineBones; ArrayIndex++)
 		{
 			const FOrientationWarpingSpineBoneData& BoneData = SpineBoneDataArray[ArrayIndex];
-			const FQuat SpineBoneCounterRotation = FQuat(RotationAxisVector, -OrientationAngle * DistributedBoneOrientationAlpha * BoneData.Weight);
+			const FQuat SpineBoneCounterRotation = FQuat(RotationAxisVector, -ActualOrientationAngle * DistributedBoneOrientationAlpha * BoneData.Weight);
 			check(BoneData.Weight > 0.f);
 
 			FTransform SpineBoneTransform(Output.Pose.GetComponentSpaceTransform(BoneData.BoneIndex));
@@ -222,7 +284,7 @@ void FAnimNode_OrientationWarping::EvaluateSkeletalControl_AnyThread(FComponentS
 	// Rotate IK Foot Root
 	if (bUpdateIKFootRoot)
 	{
-		const FQuat BoneRotation = FQuat(RotationAxisVector, OrientationAngle * IKFootRootOrientationAlpha);
+		const FQuat BoneRotation = FQuat(RotationAxisVector, ActualOrientationAngle * IKFootRootOrientationAlpha);
 
 		FTransform IKFootRootTransform(Output.Pose.GetComponentSpaceTransform(IKFootData.IKFootRootBoneIndex));
 		IKFootRootTransform.SetRotation(BoneRotation * IKFootRootTransform.GetRotation());
@@ -237,7 +299,7 @@ void FAnimNode_OrientationWarping::EvaluateSkeletalControl_AnyThread(FComponentS
 
 		if (bUpdateIKFootBones)
 		{
-			const FQuat IKFootRotation = FQuat(RotationAxisVector, -OrientationAngle * IKFootRootOrientationAlpha);
+			const FQuat IKFootRotation = FQuat(RotationAxisVector, -ActualOrientationAngle * IKFootRootOrientationAlpha);
 
 			for (int32 ArrayIndex = 0; ArrayIndex < NumIKFootBones; ArrayIndex++)
 			{
@@ -267,6 +329,11 @@ bool FAnimNode_OrientationWarping::IsValidToEvaluate(const USkeleton* Skeleton,
 		return false;
 	}
 
+	if (Mode == EWarpingEvaluationMode::Manual && UE::Anim::IsInvalidWarpingAngleDegrees(OrientationAngle, KINDA_SMALL_NUMBER))
+	{
+		return false;
+	}
+
 	if (SpineBoneDataArray.IsEmpty())
 	{
 		return false;
diff --git a/Engine/Plugins/Experimental/Animation/AnimationWarping/Source/Runtime/Private/BoneControllers/AnimNode_StrideWarping.cpp b/Engine/Plugins/Experimental/Animation/AnimationWarping/Source/Runtime/Private/BoneControllers/AnimNode_StrideWarping.cpp
--- a/Engine/Plugins/Experimental/Animation/AnimationWarping/Source/Runtime/Private/BoneControllers/AnimNode_StrideWarping.cpp
+++ b/Engine/Plugins/Experimental/Animation/AnimationWarping/Source/Runtime/Private/BoneControllers/AnimNode_StrideWarping.cpp
@@ -7,9 +7,9 @@
 DECLARE_CYCLE_STAT(TEXT("StrideWarping Eval"), STAT_StrideWarping_Eval, STATGROUP_Anim);
 
 #if ENABLE_ANIM_DEBUG
-TAutoConsoleVariable<int32> CVarAnimNodeStrideWarpingDebug(TEXT("a.AnimNode.StrideWarping.Debug"), 0, TEXT("Turn on visualization debugging for Stride Warping"));
-TAutoConsoleVariable<int32> CVarAnimNodeStrideWarpingVerbose(TEXT("a.AnimNode.StrideWarping.Verbose"), 0, TEXT("Turn on verbose graph debugging for Stride Warping"));
-TAutoConsoleVariable<int32> CVarAnimNodeStrideWarpingEnable(TEXT("a.AnimNode.StrideWarping.Enable"), 1, TEXT("Toggle Stride Warping"));
+static TAutoConsoleVariable<int32> CVarAnimNodeStrideWarpingDebug(TEXT("a.AnimNode.StrideWarping.Debug"), 0, TEXT("Turn on visualization debugging for Stride Warping"));
+static TAutoConsoleVariable<int32> CVarAnimNodeStrideWarpingVerbose(TEXT("a.AnimNode.StrideWarping.Verbose"), 0, TEXT("Turn on verbose graph debugging for Stride Warping"));
+static TAutoConsoleVariable<int32> CVarAnimNodeStrideWarpingEnable(TEXT("a.AnimNode.StrideWarping.Enable"), 1, TEXT("Toggle Stride Warping"));
 #endif
 
 void FAnimNode_StrideWarping::GatherDebugData(FNodeDebugData& DebugData)
@@ -17,13 +17,17 @@ void FAnimNode_StrideWarping::GatherDebugData(FNodeDebugData& DebugData)
 	FString DebugLine = DebugData.GetNodeName(this);
 #if ENABLE_ANIM_DEBUG
 	if (CVarAnimNodeStrideWarpingVerbose.GetValueOnAnyThread() == 1)
-	{
+	{ 
 		DebugLine += FString::Printf(TEXT("\n - Evaluation Mode: (%s)"), Mode == EWarpingEvaluationMode::Graph ? TEXT("Graph") : TEXT("Manual"));
-		DebugLine += FString::Printf(TEXT("\n - Stride Scale: (%.3fd)"), StrideScale);
-		DebugLine += FString::Printf(TEXT("\n - Stride Direction: (%s)"), *(StrideDirection.ToCompactString()));
+		DebugLine += FString::Printf(TEXT("\n - Stride Scale: (%.3fd)"), ActualStrideScale);
+		DebugLine += FString::Printf(TEXT("\n - Stride Direction: (%s)"), *(ActualStrideDirection.ToCompactString()));
 		if (Mode == EWarpingEvaluationMode::Graph)
 		{
 			DebugLine += FString::Printf(TEXT("\n - Locomotion Speed: (%.3fd)"), LocomotionSpeed);
+#if WITH_EDITORONLY_DATA
+			DebugLine += FString::Printf(TEXT("\n - Root Motion Direction: (%s)"), *(CachedRootMotionDeltaTranslation.GetSafeNormal().ToCompactString()));
+			DebugLine += FString::Printf(TEXT("\n - Root Motion Speed: (%.3fd)"), CachedRootMotionDeltaSpeed);
+#endif
 			DebugLine += FString::Printf(TEXT("\n - Min Locomotion Speed Threshold: (%.3fd)"), MinLocomotionSpeedThreshold);
 		}
 		DebugLine += FString::Printf(TEXT("\n - Floor Normal: (%s)"), *(FloorNormalDirection.Value.ToCompactString()));
@@ -32,7 +36,7 @@ void FAnimNode_StrideWarping::GatherDebugData(FNodeDebugData& DebugData)
 	else
 #endif
 	{
-		DebugLine += FString::Printf(TEXT("(Stride Scale: %.3fd, Direction: %s)"), StrideScale, *(StrideDirection.ToCompactString()));
+		DebugLine += FString::Printf(TEXT("(Stride Scale: %.3fd, Direction: %s)"), ActualStrideScale, *(ActualStrideDirection.ToCompactString()));
 	}
 	DebugData.AddDebugItem(DebugLine);
 	ComponentPose.GatherDebugData(DebugData);
@@ -41,8 +45,15 @@ void FAnimNode_StrideWarping::GatherDebugData(FNodeDebugData& DebugData)
 void FAnimNode_StrideWarping::Initialize_AnyThread(const FAnimationInitializeContext& Context)
 {
 	FAnimNode_SkeletalControlBase::Initialize_AnyThread(Context);
+	
 	AnimInstanceProxy = Context.AnimInstanceProxy;
 	StrideScaleModifierState.Reinitialize();
+	ActualStrideDirection = FVector::ForwardVector;
+	ActualStrideScale = 1.f;
+#if WITH_EDITORONLY_DATA
+	CachedRootMotionDeltaTranslation = FVector::ZeroVector;
+	CachedRootMotionDeltaSpeed = 0.f;
+#endif
 }
 
 void FAnimNode_StrideWarping::UpdateInternal(const FAnimationUpdateContext& Context)
@@ -54,19 +65,39 @@ void FAnimNode_StrideWarping::UpdateInternal(const FAnimationUpdateContext& Cont
 void FAnimNode_StrideWarping::EvaluateSkeletalControl_AnyThread(FComponentSpacePoseContext& Output, TArray<FBoneTransform>& OutBoneTransforms)
 {
 	SCOPE_CYCLE_COUNTER(STAT_StrideWarping_Eval);
-	check(OutBoneTransforms.Num() == 0);
+	check(OutBoneTransforms.IsEmpty());
 
-	FTransform RootMotionTransformDelta;
+	ActualStrideDirection = StrideDirection;
+	ActualStrideScale = StrideScale;
+
+	bool bGraphDrivenWarping = false;
 	const UE::Anim::IAnimRootMotionProvider* RootMotionProvider = UE::Anim::IAnimRootMotionProvider::Get();
-	bool bGraphDrivenWarping = RootMotionProvider && Mode == EWarpingEvaluationMode::Graph;
 
+	if (Mode == EWarpingEvaluationMode::Graph)
+	{
+		bGraphDrivenWarping = !!RootMotionProvider;
+		ensureMsgf(bGraphDrivenWarping, TEXT("Graph driven Stride Warping expected a valid root motion delta provider interface."));
+	}
+
+	FTransform RootMotionTransformDelta = FTransform::Identity;
+#if !WITH_EDITORONLY_DATA
+	FVector CachedRootMotionDeltaTranslation = FVector::ZeroVector;
+#endif
 	if (bGraphDrivenWarping)
 	{
 		// Graph driven stride warping will override the manual stride direction with the intent of the current animation sub-graph's accumulated root motion
 		bGraphDrivenWarping = RootMotionProvider->ExtractRootMotion(Output.CustomAttributes, RootMotionTransformDelta);
-		if (ensure(bGraphDrivenWarping))
+		if (ensureMsgf(bGraphDrivenWarping, TEXT(
+			"Graph driven Stride Warping expected a root motion delta to be present in the attribute stream.\n"
+			"Check that the evaluating Animation Sequence(s) have: \"Enable Root Motion\" enabled.")))
+		{
+			CachedRootMotionDeltaTranslation = RootMotionTransformDelta.GetTranslation();
+			ActualStrideDirection = CachedRootMotionDeltaTranslation.GetSafeNormal();
+		}
+		else
 		{
-			StrideDirection = RootMotionTransformDelta.GetTranslation().GetSafeNormal();
+			// Early exit on broken graph driven behavior 
+			return;
 		}
 	}
 
@@ -77,14 +108,16 @@ void FAnimNode_StrideWarping::EvaluateSkeletalControl_AnyThread(FComponentSpaceP
 
 	if (bOrientStrideDirectionUsingFloorNormal)
 	{
-		const FVector StrideWarpingAxis = ResolvedFloorNormal ^ StrideDirection;
-		StrideDirection = StrideWarpingAxis ^ ResolvedFloorNormal;
+		const FVector StrideWarpingAxis = ResolvedFloorNormal ^ ActualStrideDirection;
+		ActualStrideDirection = StrideWarpingAxis ^ ResolvedFloorNormal;
 	}
 
 #if ENABLE_ANIM_DEBUG
 	bool bDebugging = false;
 #if WITH_EDITORONLY_DATA
 	bDebugging = bDebugging || bEnableDebugDraw;
+#else
+	constexpr float DebugDrawScale = 1.f;
 #endif
 	bDebugging = bDebugging || CVarAnimNodeStrideWarpingDebug.GetValueOnAnyThread() == 1;
 	if (bDebugging)
@@ -92,8 +125,8 @@ void FAnimNode_StrideWarping::EvaluateSkeletalControl_AnyThread(FComponentSpaceP
 		// Draw Floor Normal
 		AnimInstanceProxy->AnimDrawDebugDirectionalArrow(
 			AnimInstanceProxy->GetComponentTransform().TransformPosition(IKFootRootTransform.GetLocation()),
-			AnimInstanceProxy->GetComponentTransform().TransformPosition(IKFootRootTransform.GetLocation() + ResolvedFloorNormal * 25.f),
-			40.f, FColor::Blue, false, 0.f, 2.f);
+			AnimInstanceProxy->GetComponentTransform().TransformPosition(IKFootRootTransform.GetLocation() + ResolvedFloorNormal * 25.f * DebugDrawScale),
+			40.f * DebugDrawScale, FColor::Blue, false, 0.f, 2.f * DebugDrawScale);
 	}
 #endif
 
@@ -109,42 +142,32 @@ void FAnimNode_StrideWarping::EvaluateSkeletalControl_AnyThread(FComponentSpaceP
 		)
 		{
 			const FVector FootWorldLocation = AnimInstanceProxy->GetComponentTransform().TransformPosition(Foot.IKFootBoneTransform.GetLocation());
-			AnimInstanceProxy->AnimDrawDebugSphere(FootWorldLocation, 8.f, 16, FColor::Red);
+			AnimInstanceProxy->AnimDrawDebugSphere(FootWorldLocation, 8.f * DebugDrawScale, 16, FColor::Red);
 		}
 #endif
 	}
 
 	if (bGraphDrivenWarping)
 	{
-		// Early out when stride warping will have minimal pose contribution
-		if (FMath::IsNearlyZero(CachedDeltaTime, KINDA_SMALL_NUMBER))
-		{
-			return;
-		}
-
-		const float RootMotionSpeed = RootMotionTransformDelta.GetTranslation().Size() / CachedDeltaTime;
-
-		if (FMath::IsNearlyZero(RootMotionSpeed, MinLocomotionSpeedThreshold) || FMath::IsNearlyZero(LocomotionSpeed, MinLocomotionSpeedThreshold))
-		{
-			StrideScale = 1.f;
-		}
-		else
-		{
-			// Graph driven stride scale factor will be determined by the ratio of the
-			// locomotion (capsule/physics) speed against the animation root motion speed
-			StrideScale = LocomotionSpeed / RootMotionSpeed;
-		}
+#if WITH_EDITORONLY_DATA
+		CachedRootMotionDeltaSpeed = CachedRootMotionDeltaTranslation.Size() / CachedDeltaTime;
+#else
+		const float CachedRootMotionDeltaSpeed = CachedRootMotionDeltaTranslation.Size() / CachedDeltaTime;
+#endif
+		// Graph driven stride scale factor will be determined by the ratio of the
+		// locomotion (capsule/physics) speed against the animation root motion speed
+		ActualStrideScale = LocomotionSpeed / CachedRootMotionDeltaSpeed;
 	}
 
-	// Allow the opportunity for stride scale clamping and biasing regardless of evaluation mode
-	StrideScale = StrideScaleModifierState.ApplyTo(StrideScaleModifier, StrideScale, CachedDeltaTime);
+	// Allow the opportunity for stride scale clamping and interpolation regardless of evaluation mode
+	ActualStrideScale = StrideScaleModifierState.ApplyTo(StrideScaleModifier, ActualStrideScale, CachedDeltaTime);
 
 	if (bGraphDrivenWarping)
 	{
 		// Forward the side effects of stride warping on the root motion contribution for this sub-graph
-		RootMotionTransformDelta.ScaleTranslation(StrideScale);
+		RootMotionTransformDelta.ScaleTranslation(ActualStrideScale);
 		const bool bRootMotionOverridden = RootMotionProvider->OverrideRootMotion(RootMotionTransformDelta, Output.CustomAttributes);
-		ensure(bRootMotionOverridden);
+		ensureMsgf(bRootMotionOverridden, TEXT("Graph driven Stride Warping expected a root motion delta to be present in the attribute stream prior to warping/overriding it."));
 	}
 
 	// Scale IK feet bones along Stride Warping Axis, from the Thigh bone location.
@@ -159,10 +182,10 @@ void FAnimNode_StrideWarping::EvaluateSkeletalControl_AnyThread(FComponentSpaceP
 		const FVector StrideWarpingPlaneOrigin = (FMath::Abs(ResolvedGravityDirection | ResolvedFloorNormal) > DELTA) ? FMath::LinePlaneIntersection(ThighBoneLocation, ThighBoneLocation + ResolvedGravityDirection, IKFootLocation, ResolvedFloorNormal) : IKFootLocation;
 
 		// Project FK Foot along StrideWarping Plane, this will be our Scale Origin
-		const FVector ScaleOrigin = FVector::PointPlaneProject(IKFootLocation, StrideWarpingPlaneOrigin, StrideDirection);
+		const FVector ScaleOrigin = FVector::PointPlaneProject(IKFootLocation, StrideWarpingPlaneOrigin, ActualStrideDirection);
 
 		// Now the ScaleOrigin and IKFootLocation are forming a line parallel to the floor, and we can scale the IK foot.
-		const FVector WarpedLocation = ScaleOrigin + (IKFootLocation - ScaleOrigin) * StrideScale;
+		const FVector WarpedLocation = ScaleOrigin + (IKFootLocation - ScaleOrigin) * ActualStrideScale;
 		Foot.IKFootBoneTransform.SetLocation(WarpedLocation);
 
 #if ENABLE_ANIM_DEBUG
@@ -173,17 +196,16 @@ void FAnimNode_StrideWarping::EvaluateSkeletalControl_AnyThread(FComponentSpaceP
 		)
 		{
 			const FVector FootWorldLocation = AnimInstanceProxy->GetComponentTransform().TransformPosition(Foot.IKFootBoneTransform.GetLocation());
-			AnimInstanceProxy->AnimDrawDebugSphere(FootWorldLocation, 8.f, 16, FColor::Green);
+			AnimInstanceProxy->AnimDrawDebugSphere(FootWorldLocation, 8.f * DebugDrawScale, 16, FColor::Green);
 
 			const FVector ScaleOriginWorldLoc = AnimInstanceProxy->GetComponentTransform().TransformPosition(ScaleOrigin);
-			AnimInstanceProxy->AnimDrawDebugSphere(ScaleOriginWorldLoc, 8.f, 16, FColor::Yellow);
+			AnimInstanceProxy->AnimDrawDebugSphere(ScaleOriginWorldLoc, 8.f * DebugDrawScale, 16, FColor::Yellow);
 		}
 #endif
 	}
 
 	FVector PelvisOffset = FVector::ZeroVector;
 	const FCompactPoseBoneIndex PelvisBoneIndex = PelvisBone.GetCompactPoseIndex(RequiredBones);
-	check(PelvisBoneIndex != INDEX_NONE);
 
 	FTransform PelvisTransform = Output.Pose.GetComponentSpaceTransform(PelvisBoneIndex);
 	const FVector InitialPelvisLocation = PelvisTransform.GetLocation();
@@ -214,8 +236,8 @@ void FAnimNode_StrideWarping::EvaluateSkeletalControl_AnyThread(FComponentSpaceP
 	)
 	{
 		// Draw Adjustments in Pelvis location
-		AnimInstanceProxy->AnimDrawDebugSphere(AnimInstanceProxy->GetComponentTransform().TransformPosition(InitialPelvisLocation), 8.f, 16, FColor::Red);
-		AnimInstanceProxy->AnimDrawDebugSphere(AnimInstanceProxy->GetComponentTransform().TransformPosition(PelvisTransform.GetLocation()), 8.f, 16, FColor::Blue);
+		AnimInstanceProxy->AnimDrawDebugSphere(AnimInstanceProxy->GetComponentTransform().TransformPosition(InitialPelvisLocation), 8.f * DebugDrawScale, 16, FColor::Red);
+		AnimInstanceProxy->AnimDrawDebugSphere(AnimInstanceProxy->GetComponentTransform().TransformPosition(PelvisTransform.GetLocation()), 8.f * DebugDrawScale, 16, FColor::Blue);
 	}
 
 	if (bDebugging)
@@ -223,8 +245,8 @@ void FAnimNode_StrideWarping::EvaluateSkeletalControl_AnyThread(FComponentSpaceP
 		// Draw Stride Direction
 		AnimInstanceProxy->AnimDrawDebugDirectionalArrow(
 			AnimInstanceProxy->GetComponentTransform().TransformPosition(InitialPelvisLocation),
-			AnimInstanceProxy->GetComponentTransform().TransformPosition(InitialPelvisLocation + StrideDirection * StrideScale * 100.f),
-			40.f, FColor::Red, false, 0.f, 2.f);
+			AnimInstanceProxy->GetComponentTransform().TransformPosition(InitialPelvisLocation + ActualStrideDirection * ActualStrideScale * 100.f * DebugDrawScale),
+			40.f * DebugDrawScale, FColor::Red, false, 0.f, 2.f * DebugDrawScale);
 	}
 #endif
 	// Add adjusted pelvis transform
@@ -258,12 +280,12 @@ void FAnimNode_StrideWarping::EvaluateSkeletalControl_AnyThread(FComponentSpaceP
 				AnimInstanceProxy->AnimDrawDebugLine(
 					AnimInstanceProxy->GetComponentTransform().TransformPosition(ThighTransform.GetLocation()),
 					AnimInstanceProxy->GetComponentTransform().TransformPosition(FKFootTransform.GetLocation()),
-					FColor::Red, false, 0.f, 2.f);
+					FColor::Red, false, 0.f, 2.f * DebugDrawScale);
 
 				AnimInstanceProxy->AnimDrawDebugLine(
 					AnimInstanceProxy->GetComponentTransform().TransformPosition(AdjustedThighTransform.GetLocation()),
 					AnimInstanceProxy->GetComponentTransform().TransformPosition(Foot.IKFootBoneTransform.GetLocation()),
-					FColor::Green, false, 0.f, 2.f);
+					FColor::Green, false, 0.f, 2.f * DebugDrawScale);
 			}
 #endif
 			// Find Delta Rotation take takes us from Old to New dir
@@ -300,7 +322,7 @@ void FAnimNode_StrideWarping::EvaluateSkeletalControl_AnyThread(FComponentSpaceP
 #endif
 		)
 		{
-			AnimInstanceProxy->AnimDrawDebugSphere(AnimInstanceProxy->GetComponentTransform().TransformPosition(Foot.IKFootBoneTransform.GetLocation()), 8.f, 16, FColor::Blue);
+			AnimInstanceProxy->AnimDrawDebugSphere(AnimInstanceProxy->GetComponentTransform().TransformPosition(Foot.IKFootBoneTransform.GetLocation()), 8.f * DebugDrawScale, 16, FColor::Blue);
 		}
 #endif
 		check(!Foot.IKFootBoneTransform.ContainsNaN());
@@ -319,12 +341,12 @@ bool FAnimNode_StrideWarping::IsValidToEvaluate(const USkeleton* Skeleton, const
 		return false;
 	}
 #endif
-	if (PelvisBone.GetCompactPoseIndex(RequiredBones) == INDEX_NONE)
+	if (!PelvisBone.IsValidToEvaluate(RequiredBones))
 	{
 		return false;
 	}
-	
-	if (IKFootRootBone.GetCompactPoseIndex(RequiredBones) == INDEX_NONE)
+
+	if (!IKFootRootBone.IsValidToEvaluate(RequiredBones))
 	{
 		return false;
 	}
@@ -343,7 +365,7 @@ bool FAnimNode_StrideWarping::IsValidToEvaluate(const USkeleton* Skeleton, const
 			}
 		}
 	}
-	
+
 	if (Mode == EWarpingEvaluationMode::Manual)
 	{
 		if (FMath::IsNearlyEqual(StrideScaleModifierState.ApplyTo(StrideScaleModifier, StrideScale, CachedDeltaTime), 1.f, KINDA_SMALL_NUMBER))
@@ -355,6 +377,7 @@ bool FAnimNode_StrideWarping::IsValidToEvaluate(const USkeleton* Skeleton, const
 	{
 		return false;
 	}
+
 	return true;
 }
 
diff --git a/Engine/Plugins/Experimental/Animation/AnimationWarping/Source/Runtime/Public/BoneControllers/AnimNode_OrientationWarping.h b/Engine/Plugins/Experimental/Animation/AnimationWarping/Source/Runtime/Public/BoneControllers/AnimNode_OrientationWarping.h
--- a/Engine/Plugins/Experimental/Animation/AnimationWarping/Source/Runtime/Public/BoneControllers/AnimNode_OrientationWarping.h
+++ b/Engine/Plugins/Experimental/Animation/AnimationWarping/Source/Runtime/Public/BoneControllers/AnimNode_OrientationWarping.h
@@ -56,12 +56,16 @@ struct ANIMATIONWARPINGRUNTIME_API FAnimNode_OrientationWarping : public FAnimNo
 	UPROPERTY(EditAnywhere, Category=Settings)
 	TEnumAsByte<EAxis::Type> RotationAxis = EAxis::Z;
 
-	// Specifies the interpolation speed (in degrees per second) towards reaching the final warped rotation angle
+	// Specifies the interpolation speed (in Alpha per second) towards reaching the final warped rotation angle
 	// A value of 0 will cause instantaneous warping, while a greater value will introduce smoothing
 	UPROPERTY(EditAnywhere, Category=Settings, meta=(ClampMin="0.0"))
 	float RotationInterpSpeed = 10.f;
 
 #if WITH_EDITORONLY_DATA
+	// Scale all debug drawing visualization by a factor
+	UPROPERTY(EditAnywhere, Category=Debug, meta=(ClampMin="0.0"))
+	float DebugDrawScale = 1.f;
+
 	// Enable/Disable orientation warping debug drawing
 	UPROPERTY(EditAnywhere, Category=Debug)
 	bool bEnableDebugDraw = false;
@@ -129,6 +133,12 @@ private:
 	// Computed IK bone indices for the specified foot definitions 
 	FOrientationWarpingFootData IKFootData;
 
-	// Internally cached previous frame orientation warping angle
-	float PreviousWarpedRotation = 0.f;
+	// Internal previous frame root motion delta direction
+	FVector PreviousRootMotionDeltaDirection = FVector::ZeroVector;
+
+	// Internal previous frame orientation warping angle
+	float PreviousOrientationAngle = 0.f;
+	
+	// Internal orientation warping angle
+	float ActualOrientationAngle = 0.f;
 };
\ No newline at end of file
diff --git a/Engine/Plugins/Experimental/Animation/AnimationWarping/Source/Runtime/Public/BoneControllers/AnimNode_SlopeWarping.h b/Engine/Plugins/Experimental/Animation/AnimationWarping/Source/Runtime/Public/BoneControllers/AnimNode_SlopeWarping.h
--- a/Engine/Plugins/Experimental/Animation/AnimationWarping/Source/Runtime/Public/BoneControllers/AnimNode_SlopeWarping.h
+++ b/Engine/Plugins/Experimental/Animation/AnimationWarping/Source/Runtime/Public/BoneControllers/AnimNode_SlopeWarping.h
@@ -61,7 +61,7 @@ public:
 class USkeletalMeshComponent;
 class UCharacterMovementComponent;
 
-USTRUCT()
+USTRUCT(Experimental)
 struct ANIMATIONWARPINGRUNTIME_API FAnimNode_SlopeWarping : public FAnimNode_SkeletalControlBase
 {
 	GENERATED_USTRUCT_BODY()
diff --git a/Engine/Plugins/Experimental/Animation/AnimationWarping/Source/Runtime/Public/BoneControllers/AnimNode_StrideWarping.h b/Engine/Plugins/Experimental/Animation/AnimationWarping/Source/Runtime/Public/BoneControllers/AnimNode_StrideWarping.h
--- a/Engine/Plugins/Experimental/Animation/AnimationWarping/Source/Runtime/Public/BoneControllers/AnimNode_StrideWarping.h
+++ b/Engine/Plugins/Experimental/Animation/AnimationWarping/Source/Runtime/Public/BoneControllers/AnimNode_StrideWarping.h
@@ -42,11 +42,12 @@ struct ANIMATIONWARPINGRUNTIME_API FAnimNode_StrideWarping : public FAnimNode_Sk
 
 	// Stride scale, specifying the amount of warping applied to the IK foot definitions
 	// Example: A value of 0.5 will decrease the effective leg stride by half, while a value of 2.0 will double it
-	UPROPERTY(EditAnywhere, Category=Evaluation, meta=(ClampMin="0.0", PinShownByDefault))
+	UPROPERTY(EditAnywhere, BlueprintReadWrite, Category=Evaluation, meta=(ClampMin="0.0", PinShownByDefault))
 	float StrideScale = 1.f;
 
 	// Locomotion speed, specifying the current speed of the character
 	// This will be used in the following equation for computing the stride scale: [StrideScale = (LocomotionSpeed / RootMotionSpeed)]
+	// Note: Both speeds should be relative to the delta time of the animation system
 	// 
 	// Stride scale is a value specifying the amount of warping applied to the IK foot definitions
 	// Example: A value of 0.5 will decrease the effective leg stride by half, while a value of 2.0 will double it
@@ -101,6 +102,10 @@ struct ANIMATIONWARPINGRUNTIME_API FAnimNode_StrideWarping : public FAnimNode_Sk
 	bool bClampIKUsingFKLimits = true;
 
 #if WITH_EDITORONLY_DATA
+	// Scale all debug drawing visualization by a factor
+	UPROPERTY(EditAnywhere, Category=Debug, meta=(ClampMin="0.0"))
+	float DebugDrawScale = 1.f;
+
 	// Enable/Disable stride warping debug drawing
 	UPROPERTY(EditAnywhere, Category=Debug)
 	bool bEnableDebugDraw = false;
@@ -165,9 +170,23 @@ private:
 	// Computed IK, FK, Thigh bone indices for the specified foot definitions
 	TArray<FStrideWarpingFootData> FootData;
 
+	// Internal cached stride scale modifier state
+	FInputClampState StrideScaleModifierState;
+
+	// Internal stride direction
+	FVector ActualStrideDirection = FVector::ForwardVector;
+
+	// Internal stride scale
+	float ActualStrideScale = 1.f;
+
 	// Internal cached delta time used for interpolators
 	float CachedDeltaTime = 0.f;
 
-	// Internal cached stride scale modifier state
-	FInputClampState StrideScaleModifierState;
+#if WITH_EDITORONLY_DATA
+	// Internal cached debug root motion delta translation
+	FVector CachedRootMotionDeltaTranslation = FVector::ZeroVector;
+
+	// Internal cached debug root motion speed
+	float CachedRootMotionDeltaSpeed = 0.f;
+#endif
 };
diff --git a/Engine/Source/Runtime/AnimGraphRuntime/Private/AnimNodes/AnimNode_SequenceEvaluator.cpp b/Engine/Source/Runtime/AnimGraphRuntime/Private/AnimNodes/AnimNode_SequenceEvaluator.cpp
--- a/Engine/Source/Runtime/AnimGraphRuntime/Private/AnimNodes/AnimNode_SequenceEvaluator.cpp
+++ b/Engine/Source/Runtime/AnimGraphRuntime/Private/AnimNodes/AnimNode_SequenceEvaluator.cpp
@@ -41,8 +41,6 @@ void FAnimNode_SequenceEvaluatorBase::UpdateAssetPlayer(const FAnimationUpdateCo
 
 		if ((!GetTeleportToExplicitTime() || (GetGroupName() != NAME_None) || (GetGroupMethod() == EAnimSyncMethod::Graph)) && (Context.AnimInstanceProxy->IsSkeletonCompatible(CurrentSequence->GetSkeleton())))
 		{
-			DeltaTimeRecord.SetPrevious(InternalTimeAccumulator);
-
 			if (bReinitialized)
 			{
 				switch (GetReinitializationBehavior())
@@ -63,8 +61,6 @@ void FAnimNode_SequenceEvaluatorBase::UpdateAssetPlayer(const FAnimationUpdateCo
 				InternalTimeAccumulator = CurrentExplicitTime;
 			}
 
-			DeltaTimeRecord.Delta = TimeJump;
-
 			const float DeltaTime = Context.GetDeltaTime();
 			const float RateScale = CurrentSequence->RateScale;
 			const float PlayRate = FMath::IsNearlyZero(DeltaTime) || FMath::IsNearlyZero(RateScale) ? 0.f : (TimeJump / (DeltaTime * RateScale));
