Migrating the Animation Warping plugin from Experimental to Beta.

diff --git a/Engine/Plugins/Animation/AnimationWarping/AnimationWarping.uplugin b/Engine/Plugins/Animation/AnimationWarping/AnimationWarping.uplugin
new file mode 100644
--- /dev/null
+++ b/Engine/Plugins/Animation/AnimationWarping/AnimationWarping.uplugin
@@ -0,0 +1,27 @@
+{
+	"FileVersion": 3,
+	"Version": 1,
+	"VersionName": "1.0",
+	"FriendlyName": "Animation Warping",
+	"Description": "Framework for animation and pose warping. This plugin includes Stride, Orientation, and Slope Warping alongside the Root Motion Delta animation attribute.",
+	"Category": "Animation",
+	"CreatedBy": "Epic Games, Inc.",
+	"CreatedByURL": "http://epicgames.com",
+	"EnabledByDefault": false,
+	"CanContainContent": true,
+	"IsBetaVersion": true,
+	"IsExperimentalVersion": false,
+	"Installed": false,
+	"Modules": [
+		{
+			"Name": "AnimationWarpingRuntime",
+			"Type": "Runtime",
+			"LoadingPhase": "Default"
+		},
+		{
+			"Name": "AnimationWarpingEditor",
+			"Type": "Editor",
+			"LoadingPhase": "Default"
+		},
+	]
+}
\ No newline at end of file
diff --git a/Engine/Plugins/Animation/AnimationWarping/Source/Editor/AnimationWarpingEditor.Build.cs b/Engine/Plugins/Animation/AnimationWarping/Source/Editor/AnimationWarpingEditor.Build.cs
new file mode 100644
--- /dev/null
+++ b/Engine/Plugins/Animation/AnimationWarping/Source/Editor/AnimationWarpingEditor.Build.cs
@@ -0,0 +1,41 @@
+// Copyright Epic Games, Inc. All Rights Reserved.
+
+namespace UnrealBuildTool.Rules
+{
+	public class AnimationWarpingEditor : ModuleRules
+	{
+		public AnimationWarpingEditor(ReadOnlyTargetRules Target) : base(Target)
+		{
+			PublicDependencyModuleNames.AddRange(
+			new string[]
+			{
+				"AnimationCore",
+				"AnimationWarpingRuntime",
+                "AnimGraph",
+				"AnimGraphRuntime",
+                "Core",
+                "CoreUObject",
+                "Engine",
+            });
+
+			PrivateDependencyModuleNames.AddRange(
+			new string[]
+			{
+                "SlateCore",
+            });
+
+            if (Target.bBuildEditor == true)
+            {
+                PrivateDependencyModuleNames.AddRange(
+                    new string[]
+                    {
+						"BlueprintGraph",
+						"EditorFramework",
+						"Kismet",
+                        "UnrealEd",
+                    }
+                );
+            }
+        }
+	}
+}
diff --git a/Engine/Plugins/Animation/AnimationWarping/Source/Editor/Private/AnimGraph/AnimGraphNode_OrientationWarping.cpp b/Engine/Plugins/Animation/AnimationWarping/Source/Editor/Private/AnimGraph/AnimGraphNode_OrientationWarping.cpp
new file mode 100644
--- /dev/null
+++ b/Engine/Plugins/Animation/AnimationWarping/Source/Editor/Private/AnimGraph/AnimGraphNode_OrientationWarping.cpp
@@ -0,0 +1,265 @@
+// Copyright Epic Games, Inc. All Rights Reserved.
+
+#include "AnimGraph/AnimGraphNode_OrientationWarping.h"
+#include "Animation/AnimRootMotionProvider.h"
+#include "Kismet2/BlueprintEditorUtils.h"
+#include "Kismet2/CompilerResultsLog.h"
+#include "DetailCategoryBuilder.h"
+#include "DetailLayoutBuilder.h"
+#include "ScopedTransaction.h"
+
+#define LOCTEXT_NAMESPACE "AnimationWarping"
+
+UAnimGraphNode_OrientationWarping::UAnimGraphNode_OrientationWarping(const FObjectInitializer& ObjectInitializer) : Super(ObjectInitializer)
+{
+}
+
+FText UAnimGraphNode_OrientationWarping::GetControllerDescription() const
+{
+	return LOCTEXT("OrientationWarping", "Orientation Warping");
+}
+
+FText UAnimGraphNode_OrientationWarping::GetTooltipText() const
+{
+	return LOCTEXT("OrientationWarpingTooltip", "Warps the root bone, and counter rotates the body, to match locomotion direction.");
+}
+
+FText UAnimGraphNode_OrientationWarping::GetNodeTitle(ENodeTitleType::Type TitleType) const
+{
+	return GetControllerDescription();
+}
+
+FLinearColor UAnimGraphNode_OrientationWarping::GetNodeTitleColor() const
+{
+	return FLinearColor(FColor(153.f, 0.f, 0.f));
+}
+
+void UAnimGraphNode_OrientationWarping::CustomizePinData(UEdGraphPin* Pin, FName SourcePropertyName, int32 ArrayIndex) const
+{
+	Super::CustomizePinData(Pin, SourcePropertyName, ArrayIndex);
+
+	if (Pin->PinName == GET_MEMBER_NAME_STRING_CHECKED(FAnimNode_OrientationWarping, OrientationAngle))
+	{
+		Pin->bHidden = (Node.Mode == EWarpingEvaluationMode::Graph);
+	}
+
+	if (Pin->PinName == GET_MEMBER_NAME_STRING_CHECKED(FAnimNode_OrientationWarping, LocomotionAngle))
+	{
+		Pin->bHidden = (Node.Mode == EWarpingEvaluationMode::Manual);
+	}
+
+	if (Pin->PinName == GET_MEMBER_NAME_STRING_CHECKED(FAnimNode_OrientationWarping, LocomotionAngleDeltaThreshold))
+	{
+		Pin->bHidden = (Node.Mode == EWarpingEvaluationMode::Manual);
+	}
+}
+
+void UAnimGraphNode_OrientationWarping::CustomizeDetails(IDetailLayoutBuilder& DetailBuilder)
+{
+	DECLARE_SCOPE_HIERARCHICAL_COUNTER_FUNC()
+	Super::CustomizeDetails(DetailBuilder);
+
+	DetailBuilder.SortCategories([](const TMap<FName, IDetailCategoryBuilder*>& CategoryMap)
+	{
+		for (const TPair<FName, IDetailCategoryBuilder*>& Pair : CategoryMap)
+		{
+			int32 SortOrder = Pair.Value->GetSortOrder();
+			const FName CategoryName = Pair.Key;
+
+			if (CategoryName == "Evaluation")
+			{
+				SortOrder += 1;
+			}
+			else if (CategoryName == "Settings")
+			{
+				SortOrder += 2;
+			}
+			else if (CategoryName == "Debug")
+			{
+				SortOrder += 3;
+			}
+			else
+			{
+				const int32 ValueSortOrder = Pair.Value->GetSortOrder();
+				if (ValueSortOrder >= SortOrder && ValueSortOrder < SortOrder + 10)
+				{
+					SortOrder += 10;
+				}
+				else
+				{
+					continue;
+				}
+			}
+
+			Pair.Value->SetSortOrder(SortOrder);
+		}
+	});
+
+	TSharedRef<IPropertyHandle> NodeHandle = DetailBuilder.GetProperty(FName(TEXT("Node")), GetClass());
+
+	if (Node.Mode == EWarpingEvaluationMode::Graph)
+	{
+		DetailBuilder.HideProperty(NodeHandle->GetChildHandle(GET_MEMBER_NAME_CHECKED(FAnimNode_OrientationWarping, OrientationAngle)));
+	}
+
+	if (Node.Mode == EWarpingEvaluationMode::Manual)
+	{
+		DetailBuilder.HideProperty(NodeHandle->GetChildHandle(GET_MEMBER_NAME_CHECKED(FAnimNode_OrientationWarping, LocomotionAngle)));
+		DetailBuilder.HideProperty(NodeHandle->GetChildHandle(GET_MEMBER_NAME_CHECKED(FAnimNode_OrientationWarping, LocomotionAngleDeltaThreshold)));
+	}
+}
+
+void UAnimGraphNode_OrientationWarping::PostEditChangeProperty(struct FPropertyChangedEvent& PropertyChangedEvent)
+{
+	DECLARE_SCOPE_HIERARCHICAL_COUNTER_FUNC()
+	Super::PostEditChangeProperty(PropertyChangedEvent);
+
+	bool bRequiresNodeReconstruct = false;
+	FProperty* ChangedProperty = PropertyChangedEvent.Property;
+
+	if (ChangedProperty)
+	{
+		// Evaluation mode
+		if (ChangedProperty->GetFName() == GET_MEMBER_NAME_STRING_CHECKED(FAnimNode_OrientationWarping, Mode))
+		{
+			FScopedTransaction Transaction(LOCTEXT("ChangeEvaluationMode", "Change Evaluation Mode"));
+			Modify();
+
+			// Break links to pins going away
+			for (int32 PinIndex = 0; PinIndex < Pins.Num(); ++PinIndex)
+			{
+				UEdGraphPin* Pin = Pins[PinIndex];
+				if (Pin->PinName == GET_MEMBER_NAME_STRING_CHECKED(FAnimNode_OrientationWarping, OrientationAngle))
+				{
+					if (Node.Mode == EWarpingEvaluationMode::Graph)
+					{
+						Pin->BreakAllPinLinks();
+					}
+				}
+				else if (Pin->PinName == GET_MEMBER_NAME_STRING_CHECKED(FAnimNode_OrientationWarping, LocomotionAngle))
+				{
+					if (Node.Mode == EWarpingEvaluationMode::Manual)
+					{
+						Pin->BreakAllPinLinks();
+					}
+				}
+				else if (Pin->PinName == GET_MEMBER_NAME_STRING_CHECKED(FAnimNode_OrientationWarping, LocomotionAngleDeltaThreshold))
+				{
+					if (Node.Mode == EWarpingEvaluationMode::Manual)
+					{
+						Pin->BreakAllPinLinks();
+					}
+				}
+			}
+
+			bRequiresNodeReconstruct = true;
+		}
+	}
+
+	if (bRequiresNodeReconstruct)
+	{
+		ReconstructNode();
+		FBlueprintEditorUtils::MarkBlueprintAsStructurallyModified(GetBlueprint());
+	}
+}
+
+void UAnimGraphNode_OrientationWarping::GetInputLinkAttributes(FNodeAttributeArray& OutAttributes) const
+{
+	if (Node.Mode == EWarpingEvaluationMode::Graph)
+	{
+		OutAttributes.Add(UE::Anim::IAnimRootMotionProvider::AttributeName);
+	}
+}
+
+void UAnimGraphNode_OrientationWarping::GetOutputLinkAttributes(FNodeAttributeArray& OutAttributes) const
+{
+	if (Node.Mode == EWarpingEvaluationMode::Graph)
+	{
+		OutAttributes.Add(UE::Anim::IAnimRootMotionProvider::AttributeName);
+	}
+}
+
+void UAnimGraphNode_OrientationWarping::ValidateAnimNodeDuringCompilation(USkeleton* ForSkeleton, FCompilerResultsLog& MessageLog)
+{
+	auto HasInvalidBoneName = [](const FName& BoneName)
+	{
+		return BoneName == NAME_None;
+	};
+
+	auto HasInvalidBoneIndex = [&](const FName& BoneName)
+	{
+		return ForSkeleton->GetReferenceSkeleton().FindBoneIndex(BoneName) == INDEX_NONE;
+	};
+
+	auto InvalidBoneNameMessage = [&](const FName& BoneName)
+	{
+		FFormatNamedArguments Args;
+		Args.Add(TEXT("BoneName"), FText::FromName(BoneName));
+		const FText Message = FText::Format(NSLOCTEXT("OrientationWarping", "Invalid{BoneName}BoneName", "@@ - {BoneName} bone not found in Skeleton"), Args);
+		MessageLog.Warning(*Message.ToString(), this);
+	};
+
+	auto InvalidBoneIndexMessage = [&](const FName& BoneName)
+	{
+		FFormatNamedArguments Args;
+		Args.Add(TEXT("BoneName"), FText::FromName(BoneName));
+		const FText Message = FText::Format(NSLOCTEXT("OrientationWarping", "Invalid{BoneName}BoneInSkeleton", "@@ - {BoneName} bone definition is required"), Args);
+		MessageLog.Warning(*Message.ToString(), this);
+	};
+
+	if (Node.RotationAxis == EAxis::None)
+	{
+		MessageLog.Warning(*NSLOCTEXT("OrientationWarping", "InvalidRotationAxis", "@@ - Rotation Axis choice of X, Y, or Z is required").ToString(), this);
+	}
+
+	if (Node.SpineBones.IsEmpty())
+	{
+		MessageLog.Warning(*NSLOCTEXT("OrientationWarping", "InvalidSpineBones", "@@ - Spine bone definitions are required").ToString(), this);
+	}
+	else
+	{
+		for (const auto& Bone : Node.SpineBones)
+		{
+			if (HasInvalidBoneName(Bone.BoneName))
+			{
+				InvalidBoneIndexMessage("Spine");
+			}
+			else if (HasInvalidBoneIndex(Bone.BoneName))
+			{
+				InvalidBoneNameMessage(Bone.BoneName);
+			}
+		}
+	}
+
+	if (HasInvalidBoneName(Node.IKFootRootBone.BoneName))
+	{
+		InvalidBoneIndexMessage("IK Foot Root");
+	}
+	else if (HasInvalidBoneIndex(Node.IKFootRootBone.BoneName))
+	{
+		InvalidBoneNameMessage(Node.IKFootRootBone.BoneName);
+	}
+
+	if (Node.SpineBones.IsEmpty())
+	{
+		MessageLog.Warning(*NSLOCTEXT("OrientationWarping", "InvalidIKFootBones", "@@ - IK Foot bone definitions are required").ToString(), this);
+	}
+	else
+	{
+		for (const auto& Bone : Node.IKFootBones)
+		{
+			if (HasInvalidBoneName(Bone.BoneName))
+			{
+				InvalidBoneIndexMessage("IK Foot");
+			}
+			else if (HasInvalidBoneIndex(Bone.BoneName))
+			{
+				InvalidBoneNameMessage(Bone.BoneName);
+			}
+		}
+	}
+
+	Super::ValidateAnimNodeDuringCompilation(ForSkeleton, MessageLog);
+}
+
+#undef LOCTEXT_NAMESPACE
diff --git a/Engine/Plugins/Animation/AnimationWarping/Source/Editor/Private/AnimGraph/AnimGraphNode_SlopeWarping.cpp b/Engine/Plugins/Animation/AnimationWarping/Source/Editor/Private/AnimGraph/AnimGraphNode_SlopeWarping.cpp
new file mode 100644
--- /dev/null
+++ b/Engine/Plugins/Animation/AnimationWarping/Source/Editor/Private/AnimGraph/AnimGraphNode_SlopeWarping.cpp
@@ -0,0 +1,31 @@
+// Copyright Epic Games, Inc. All Rights Reserved.
+
+#include "AnimGraph/AnimGraphNode_SlopeWarping.h"
+
+#define LOCTEXT_NAMESPACE "AnimationWarping"
+
+UAnimGraphNode_SlopeWarping::UAnimGraphNode_SlopeWarping(const FObjectInitializer& ObjectInitializer) : Super(ObjectInitializer)
+{
+}
+
+FText UAnimGraphNode_SlopeWarping::GetControllerDescription() const
+{
+	return LOCTEXT("SlopeWarping", "Slope Warping");
+}
+
+FText UAnimGraphNode_SlopeWarping::GetTooltipText() const
+{
+	return LOCTEXT("SlopeWarpingTooltip", "Warps the feet to match the floor normal.");
+}
+
+FText UAnimGraphNode_SlopeWarping::GetNodeTitle(ENodeTitleType::Type TitleType) const
+{
+	return GetControllerDescription();
+}
+
+FLinearColor UAnimGraphNode_SlopeWarping::GetNodeTitleColor() const
+{
+	return FLinearColor(FColor(153.f, 0.f, 0.f));
+}
+
+#undef LOCTEXT_NAMESPACE
diff --git a/Engine/Plugins/Animation/AnimationWarping/Source/Editor/Private/AnimGraph/AnimGraphNode_StrideWarping.cpp b/Engine/Plugins/Animation/AnimationWarping/Source/Editor/Private/AnimGraph/AnimGraphNode_StrideWarping.cpp
new file mode 100644
--- /dev/null
+++ b/Engine/Plugins/Animation/AnimationWarping/Source/Editor/Private/AnimGraph/AnimGraphNode_StrideWarping.cpp
@@ -0,0 +1,300 @@
+// Copyright Epic Games, Inc. All Rights Reserved.
+
+#include "AnimGraph/AnimGraphNode_StrideWarping.h"
+#include "Animation/AnimRootMotionProvider.h"
+#include "Kismet2/BlueprintEditorUtils.h"
+#include "Kismet2/CompilerResultsLog.h"
+#include "DetailCategoryBuilder.h"
+#include "DetailLayoutBuilder.h"
+#include "ScopedTransaction.h"
+
+#define LOCTEXT_NAMESPACE "AnimationWarping"
+
+UAnimGraphNode_StrideWarping::UAnimGraphNode_StrideWarping(const FObjectInitializer& ObjectInitializer) : Super(ObjectInitializer)
+{
+}
+
+FText UAnimGraphNode_StrideWarping::GetControllerDescription() const
+{
+	return LOCTEXT("StrideWarping", "Stride Warping");
+}
+
+FText UAnimGraphNode_StrideWarping::GetTooltipText() const
+{
+	return LOCTEXT("StrideWarpingTooltip", "Warps the pelvis and feet to match the locomotion speed.");
+}
+
+FText UAnimGraphNode_StrideWarping::GetNodeTitle(ENodeTitleType::Type TitleType) const
+{
+	return GetControllerDescription();
+}
+
+FLinearColor UAnimGraphNode_StrideWarping::GetNodeTitleColor() const
+{
+	return FLinearColor(FColor(153.f, 0.f, 0.f));
+}
+
+void UAnimGraphNode_StrideWarping::CustomizePinData(UEdGraphPin* Pin, FName SourcePropertyName, int32 ArrayIndex) const
+{
+	Super::CustomizePinData(Pin, SourcePropertyName, ArrayIndex);
+
+	if (Pin->PinName == GET_MEMBER_NAME_STRING_CHECKED(FAnimNode_StrideWarping, StrideDirection))
+	{
+		Pin->bHidden = (Node.Mode == EWarpingEvaluationMode::Graph);
+	}
+
+	if (Pin->PinName == GET_MEMBER_NAME_STRING_CHECKED(FAnimNode_StrideWarping, StrideScale))
+	{
+		Pin->bHidden = (Node.Mode == EWarpingEvaluationMode::Graph);
+	}
+
+	if (Pin->PinName == GET_MEMBER_NAME_STRING_CHECKED(FAnimNode_StrideWarping, LocomotionSpeed))
+	{
+		Pin->bHidden = (Node.Mode == EWarpingEvaluationMode::Manual);
+	}
+
+	if (Pin->PinName == GET_MEMBER_NAME_STRING_CHECKED(FAnimNode_StrideWarping, MinLocomotionSpeedThreshold))
+	{
+		Pin->bHidden = (Node.Mode == EWarpingEvaluationMode::Manual);
+	}
+}
+
+void UAnimGraphNode_StrideWarping::CustomizeDetails(IDetailLayoutBuilder& DetailBuilder)
+{
+	DECLARE_SCOPE_HIERARCHICAL_COUNTER_FUNC()
+	Super::CustomizeDetails(DetailBuilder);
+
+	DetailBuilder.SortCategories([](const TMap<FName, IDetailCategoryBuilder*>& CategoryMap)
+	{
+		for (const TPair<FName, IDetailCategoryBuilder*>& Pair : CategoryMap)
+		{
+			int32 SortOrder = Pair.Value->GetSortOrder();
+			const FName CategoryName = Pair.Key;
+
+			if (CategoryName == "Evaluation")
+			{
+				SortOrder += 1;
+			}
+			else if (CategoryName == "Settings")
+			{
+				SortOrder += 2;
+			}
+			else if (CategoryName == "Advanced")
+			{
+				SortOrder += 3;
+			}
+			else if (CategoryName == "Debug")
+			{
+				SortOrder += 4;
+			}
+			else
+			{
+				const int32 ValueSortOrder = Pair.Value->GetSortOrder();
+				if (ValueSortOrder >= SortOrder && ValueSortOrder < SortOrder + 10)
+				{
+					SortOrder += 10;
+				}
+				else
+				{
+					continue;
+				}
+			}
+
+			Pair.Value->SetSortOrder(SortOrder);
+		}
+	});
+
+	TSharedRef<IPropertyHandle> NodeHandle = DetailBuilder.GetProperty(FName(TEXT("Node")), GetClass());
+
+	if (Node.Mode == EWarpingEvaluationMode::Graph)
+	{
+		DetailBuilder.HideProperty(NodeHandle->GetChildHandle(GET_MEMBER_NAME_CHECKED(FAnimNode_StrideWarping, StrideScale)));
+		DetailBuilder.HideProperty(NodeHandle->GetChildHandle(GET_MEMBER_NAME_CHECKED(FAnimNode_StrideWarping, StrideDirection)));
+	}
+
+	if (Node.Mode == EWarpingEvaluationMode::Manual)
+	{
+		DetailBuilder.HideProperty(NodeHandle->GetChildHandle(GET_MEMBER_NAME_CHECKED(FAnimNode_StrideWarping, LocomotionSpeed)));
+		DetailBuilder.HideProperty(NodeHandle->GetChildHandle(GET_MEMBER_NAME_CHECKED(FAnimNode_StrideWarping, MinLocomotionSpeedThreshold)));
+	}
+}
+
+void UAnimGraphNode_StrideWarping::PostEditChangeProperty(struct FPropertyChangedEvent& PropertyChangedEvent)
+{
+	DECLARE_SCOPE_HIERARCHICAL_COUNTER_FUNC()
+	Super::PostEditChangeProperty(PropertyChangedEvent);
+
+	bool bRequiresNodeReconstruct = false;
+	FProperty* ChangedProperty = PropertyChangedEvent.Property;
+
+	if (ChangedProperty)
+	{
+		if ((ChangedProperty->GetFName() == GET_MEMBER_NAME_STRING_CHECKED(FInputScaleBiasClamp, bMapRange))
+			|| (ChangedProperty->GetFName() == GET_MEMBER_NAME_STRING_CHECKED(FInputRange, Min))
+			|| (ChangedProperty->GetFName() == GET_MEMBER_NAME_STRING_CHECKED(FInputRange, Max))
+			|| (ChangedProperty->GetFName() == GET_MEMBER_NAME_STRING_CHECKED(FInputScaleBiasClamp, Scale))
+			|| (ChangedProperty->GetFName() == GET_MEMBER_NAME_STRING_CHECKED(FInputScaleBiasClamp, Bias))
+			|| (ChangedProperty->GetFName() == GET_MEMBER_NAME_STRING_CHECKED(FInputScaleBiasClamp, bClampResult))
+			|| (ChangedProperty->GetFName() == GET_MEMBER_NAME_STRING_CHECKED(FInputScaleBiasClamp, ClampMin))
+			|| (ChangedProperty->GetFName() == GET_MEMBER_NAME_STRING_CHECKED(FInputScaleBiasClamp, ClampMax))
+			|| (ChangedProperty->GetFName() == GET_MEMBER_NAME_STRING_CHECKED(FInputScaleBiasClamp, bInterpResult))
+			|| (ChangedProperty->GetFName() == GET_MEMBER_NAME_STRING_CHECKED(FInputScaleBiasClamp, InterpSpeedIncreasing))
+			|| (ChangedProperty->GetFName() == GET_MEMBER_NAME_STRING_CHECKED(FInputScaleBiasClamp, InterpSpeedDecreasing)))
+		{
+			bRequiresNodeReconstruct = true;
+		}
+
+		// Evaluation mode
+		if (ChangedProperty->GetFName() == GET_MEMBER_NAME_STRING_CHECKED(FAnimNode_StrideWarping, Mode))
+		{
+			FScopedTransaction Transaction(LOCTEXT("ChangeEvaluationMode", "Change Evaluation Mode"));
+			Modify();
+
+			// Break links to pins going away
+			for (int32 PinIndex = 0; PinIndex < Pins.Num(); ++PinIndex)
+			{
+				UEdGraphPin* Pin = Pins[PinIndex];
+				if (Pin->PinName == GET_MEMBER_NAME_STRING_CHECKED(FAnimNode_StrideWarping, StrideDirection))
+				{
+					if (Node.Mode == EWarpingEvaluationMode::Graph)
+					{
+						Pin->BreakAllPinLinks();
+					}
+				}
+				else if (Pin->PinName == GET_MEMBER_NAME_STRING_CHECKED(FAnimNode_StrideWarping, StrideScale))
+				{
+					if (Node.Mode == EWarpingEvaluationMode::Graph)
+					{
+						Pin->BreakAllPinLinks();
+					}
+				}
+				else if (Pin->PinName == GET_MEMBER_NAME_STRING_CHECKED(FAnimNode_StrideWarping, LocomotionSpeed))
+				{
+					if (Node.Mode == EWarpingEvaluationMode::Manual)
+					{
+						Pin->BreakAllPinLinks();
+					}
+				}
+				else if (Pin->PinName == GET_MEMBER_NAME_STRING_CHECKED(FAnimNode_StrideWarping, MinLocomotionSpeedThreshold))
+				{
+					if (Node.Mode == EWarpingEvaluationMode::Manual)
+					{
+						Pin->BreakAllPinLinks();
+					}
+				}
+			}
+
+			bRequiresNodeReconstruct = true;
+		}
+	}
+
+	if (bRequiresNodeReconstruct)
+	{
+		ReconstructNode();
+		FBlueprintEditorUtils::MarkBlueprintAsStructurallyModified(GetBlueprint());
+	}
+}
+
+void UAnimGraphNode_StrideWarping::GetInputLinkAttributes(FNodeAttributeArray& OutAttributes) const
+{
+	if (Node.Mode == EWarpingEvaluationMode::Graph)
+	{
+		OutAttributes.Add(UE::Anim::IAnimRootMotionProvider::AttributeName);
+	}
+}
+
+void UAnimGraphNode_StrideWarping::GetOutputLinkAttributes(FNodeAttributeArray& OutAttributes) const
+{
+	if (Node.Mode == EWarpingEvaluationMode::Graph)
+	{
+		OutAttributes.Add(UE::Anim::IAnimRootMotionProvider::AttributeName);
+	}
+}
+
+void UAnimGraphNode_StrideWarping::ValidateAnimNodeDuringCompilation(USkeleton* ForSkeleton, FCompilerResultsLog& MessageLog)
+{
+	auto HasInvalidBoneName = [](const FName& BoneName) 
+	{ 
+		return BoneName == NAME_None; 
+	};
+
+	auto HasInvalidBoneIndex = [&] (const FName& BoneName) 
+	{ 
+		return ForSkeleton->GetReferenceSkeleton().FindBoneIndex(BoneName) == INDEX_NONE; 
+	};
+
+	auto InvalidBoneNameMessage = [&](const FName& BoneName) 
+	{
+		FFormatNamedArguments Args;
+		Args.Add(TEXT("BoneName"), FText::FromName(BoneName));
+		const FText Message = FText::Format(NSLOCTEXT("StrideWarping", "Invalid{BoneName}BoneName", "@@ - {BoneName} bone not found in Skeleton"), Args);
+		MessageLog.Warning(*Message.ToString(), this);
+	};
+
+	auto InvalidBoneIndexMessage = [&](const FName& BoneName) 
+	{
+		FFormatNamedArguments Args;
+		Args.Add(TEXT("BoneName"), FText::FromName(BoneName));
+		const FText Message = FText::Format(NSLOCTEXT("StrideWarping", "Invalid{BoneName}BoneInSkeleton", "@@ - {BoneName} bone definition is required"), Args);
+		MessageLog.Warning(*Message.ToString(), this);
+	};
+
+	if (HasInvalidBoneName(Node.PelvisBone.BoneName))
+	{
+		InvalidBoneIndexMessage("Pelvis");
+	}
+	else if (HasInvalidBoneIndex(Node.PelvisBone.BoneName))
+	{
+		InvalidBoneNameMessage(Node.PelvisBone.BoneName);
+	}
+
+	if (HasInvalidBoneName(Node.IKFootRootBone.BoneName))
+	{
+		InvalidBoneIndexMessage("IK Foot Root");
+	}
+	else if (HasInvalidBoneIndex(Node.IKFootRootBone.BoneName))
+	{
+		InvalidBoneNameMessage(Node.IKFootRootBone.BoneName);
+	}
+
+	if (Node.FootDefinitions.IsEmpty())
+	{
+		MessageLog.Warning(*NSLOCTEXT("StrideWarping", "InvalidFootDefinitions", "@@ - Foot definitions are required").ToString(), this);
+	}
+	else
+	{
+		for (const auto& Foot : Node.FootDefinitions)
+		{
+			if (HasInvalidBoneName(Foot.IKFootBone.BoneName))
+			{
+				InvalidBoneIndexMessage("IK Foot");
+			}
+			else if (HasInvalidBoneIndex(Foot.IKFootBone.BoneName))
+			{
+				InvalidBoneNameMessage(Foot.IKFootBone.BoneName);
+			}
+
+			if (HasInvalidBoneName(Foot.FKFootBone.BoneName))
+			{
+				InvalidBoneIndexMessage("FK Foot");
+			}
+			else if (HasInvalidBoneIndex(Foot.FKFootBone.BoneName))
+			{
+				InvalidBoneNameMessage(Foot.FKFootBone.BoneName);
+			}
+
+			if (HasInvalidBoneName(Foot.ThighBone.BoneName))
+			{
+				InvalidBoneIndexMessage("Thigh");
+			}
+			else if (HasInvalidBoneIndex(Foot.ThighBone.BoneName))
+			{
+				InvalidBoneNameMessage(Foot.ThighBone.BoneName);
+			}
+		}
+	}
+
+	Super::ValidateAnimNodeDuringCompilation(ForSkeleton, MessageLog);
+}
+
+#undef LOCTEXT_NAMESPACE
diff --git a/Engine/Plugins/Animation/AnimationWarping/Source/Editor/Private/AnimationWarpingEditorModule.cpp b/Engine/Plugins/Animation/AnimationWarping/Source/Editor/Private/AnimationWarpingEditorModule.cpp
new file mode 100644
--- /dev/null
+++ b/Engine/Plugins/Animation/AnimationWarping/Source/Editor/Private/AnimationWarpingEditorModule.cpp
@@ -0,0 +1,11 @@
+// Copyright Epic Games, Inc. All Rights Reserved.
+
+#include "AnimationWarpingEditorModule.h"
+#include "Modules/ModuleManager.h"
+
+class FAnimationWarpingEditorModule : public IAnimationWarpingEditorModule
+{
+};
+
+
+IMPLEMENT_MODULE( FAnimationWarpingEditorModule, AnimationWarpingEditor )
diff --git a/Engine/Plugins/Animation/AnimationWarping/Source/Editor/Public/AnimGraph/AnimGraphNode_OrientationWarping.h b/Engine/Plugins/Animation/AnimationWarping/Source/Editor/Public/AnimGraph/AnimGraphNode_OrientationWarping.h
new file mode 100644
--- /dev/null
+++ b/Engine/Plugins/Animation/AnimationWarping/Source/Editor/Public/AnimGraph/AnimGraphNode_OrientationWarping.h
@@ -0,0 +1,37 @@
+// Copyright Epic Games, Inc. All Rights Reserved.
+
+#pragma once
+#include "AnimGraphNode_SkeletalControlBase.h"
+#include "BoneControllers/AnimNode_OrientationWarping.h"
+#include "AnimGraphNode_OrientationWarping.generated.h"
+
+UCLASS()
+class ANIMATIONWARPINGEDITOR_API UAnimGraphNode_OrientationWarping : public UAnimGraphNode_SkeletalControlBase
+{
+	GENERATED_UCLASS_BODY()
+
+	UPROPERTY(EditAnywhere, Category = Settings)
+	FAnimNode_OrientationWarping Node;
+
+public:
+	// UEdGraphNode interface
+	virtual FText GetNodeTitle(ENodeTitleType::Type TitleType) const override;
+	virtual FLinearColor GetNodeTitleColor() const override;
+	virtual FText GetTooltipText() const override;
+	virtual void CustomizeDetails(IDetailLayoutBuilder& DetailBuilder) override;
+	virtual void PostEditChangeProperty(struct FPropertyChangedEvent& PropertyChangedEvent) override;
+	virtual void GetInputLinkAttributes(FNodeAttributeArray& OutAttributes) const override;
+	virtual void GetOutputLinkAttributes(FNodeAttributeArray& OutAttributes) const override;
+	virtual void ValidateAnimNodeDuringCompilation(USkeleton* ForSkeleton, FCompilerResultsLog& MessageLog) override;
+	// End of UEdGraphNode interface
+
+protected:
+	// UAnimGraphNode_Base interface
+	virtual void CustomizePinData(UEdGraphPin* Pin, FName SourcePropertyName, int32 ArrayIndex) const override;
+	// End of UAnimGraphNode_Base interface
+
+	// UAnimGraphNode_SkeletalControlBase interface
+	virtual FText GetControllerDescription() const override;
+	virtual const FAnimNode_SkeletalControlBase* GetNode() const override { return &Node; }
+	// End of UAnimGraphNode_SkeletalControlBase interface
+};
diff --git a/Engine/Plugins/Animation/AnimationWarping/Source/Editor/Public/AnimGraph/AnimGraphNode_SlopeWarping.h b/Engine/Plugins/Animation/AnimationWarping/Source/Editor/Public/AnimGraph/AnimGraphNode_SlopeWarping.h
new file mode 100644
--- /dev/null
+++ b/Engine/Plugins/Animation/AnimationWarping/Source/Editor/Public/AnimGraph/AnimGraphNode_SlopeWarping.h
@@ -0,0 +1,33 @@
+// Copyright Epic Games, Inc. All Rights Reserved.
+
+#pragma once
+
+#include "AnimGraphNode_SkeletalControlBase.h"
+#include "BoneControllers/AnimNode_SlopeWarping.h"
+
+#include "AnimGraphNode_SlopeWarping.generated.h"
+
+UCLASS(Experimental)
+class ANIMATIONWARPINGEDITOR_API UAnimGraphNode_SlopeWarping : public UAnimGraphNode_SkeletalControlBase
+{
+	GENERATED_UCLASS_BODY()
+
+	UPROPERTY(EditAnywhere, Category = Settings)
+	FAnimNode_SlopeWarping Node;
+
+public:
+	// UEdGraphNode interface
+	virtual FText GetNodeTitle(ENodeTitleType::Type TitleType) const override;
+	virtual FText GetTooltipText() const override;
+	virtual FLinearColor GetNodeTitleColor() const override;
+	// End of UEdGraphNode interface
+
+protected:
+	// UAnimGraphNode_Base interface
+	// End of UAnimGraphNode_Base interface
+
+	// UAnimGraphNode_SkeletalControlBase interface
+	virtual FText GetControllerDescription() const override;
+	virtual const FAnimNode_SkeletalControlBase* GetNode() const override { return &Node; }
+	// End of UAnimGraphNode_SkeletalControlBase interface
+};
diff --git a/Engine/Plugins/Animation/AnimationWarping/Source/Editor/Public/AnimGraph/AnimGraphNode_StrideWarping.h b/Engine/Plugins/Animation/AnimationWarping/Source/Editor/Public/AnimGraph/AnimGraphNode_StrideWarping.h
new file mode 100644
--- /dev/null
+++ b/Engine/Plugins/Animation/AnimationWarping/Source/Editor/Public/AnimGraph/AnimGraphNode_StrideWarping.h
@@ -0,0 +1,39 @@
+// Copyright Epic Games, Inc. All Rights Reserved.
+
+#pragma once
+
+#include "AnimGraphNode_SkeletalControlBase.h"
+#include "BoneControllers/AnimNode_StrideWarping.h"
+
+#include "AnimGraphNode_StrideWarping.generated.h"
+
+UCLASS()
+class ANIMATIONWARPINGEDITOR_API UAnimGraphNode_StrideWarping : public UAnimGraphNode_SkeletalControlBase
+{
+	GENERATED_UCLASS_BODY()
+
+	UPROPERTY(EditAnywhere, Category = Settings)
+	FAnimNode_StrideWarping Node;
+
+public:
+	// UEdGraphNode interface
+	virtual FText GetNodeTitle(ENodeTitleType::Type TitleType) const override;
+	virtual FLinearColor GetNodeTitleColor() const override;
+	virtual FText GetTooltipText() const override;
+	virtual void CustomizeDetails(IDetailLayoutBuilder& DetailBuilder) override;
+	virtual void PostEditChangeProperty(struct FPropertyChangedEvent& PropertyChangedEvent) override;
+	virtual void GetInputLinkAttributes(FNodeAttributeArray& OutAttributes) const override;
+	virtual void GetOutputLinkAttributes(FNodeAttributeArray& OutAttributes) const override;
+	virtual void ValidateAnimNodeDuringCompilation(USkeleton* ForSkeleton, FCompilerResultsLog& MessageLog) override;
+	// End of UEdGraphNode interface
+
+protected:
+	// UAnimGraphNode_Base interface
+	virtual void CustomizePinData(UEdGraphPin* Pin, FName SourcePropertyName, int32 ArrayIndex) const override;
+	// End of UAnimGraphNode_Base interface
+
+	// UAnimGraphNode_SkeletalControlBase interface
+	virtual FText GetControllerDescription() const override;
+	virtual const FAnimNode_SkeletalControlBase* GetNode() const override { return &Node; }
+	// End of UAnimGraphNode_SkeletalControlBase interface
+};
diff --git a/Engine/Plugins/Animation/AnimationWarping/Source/Editor/Public/AnimationWarpingEditorModule.h b/Engine/Plugins/Animation/AnimationWarping/Source/Editor/Public/AnimationWarpingEditorModule.h
new file mode 100644
--- /dev/null
+++ b/Engine/Plugins/Animation/AnimationWarping/Source/Editor/Public/AnimationWarpingEditorModule.h
@@ -0,0 +1,12 @@
+// Copyright Epic Games, Inc. All Rights Reserved.
+
+#pragma once
+
+#include "CoreMinimal.h"
+#include "Modules/ModuleInterface.h"
+#include "Modules/ModuleManager.h"
+
+class IAnimationWarpingEditorModule : public IModuleInterface
+{
+};
+
diff --git a/Engine/Plugins/Animation/AnimationWarping/Source/Runtime/AnimationWarpingRuntime.Build.cs b/Engine/Plugins/Animation/AnimationWarping/Source/Runtime/AnimationWarpingRuntime.Build.cs
new file mode 100644
--- /dev/null
+++ b/Engine/Plugins/Animation/AnimationWarping/Source/Runtime/AnimationWarpingRuntime.Build.cs
@@ -0,0 +1,20 @@
+// Copyright Epic Games, Inc. All Rights Reserved.
+
+namespace UnrealBuildTool.Rules
+{
+	public class AnimationWarpingRuntime : ModuleRules
+	{
+		public AnimationWarpingRuntime(ReadOnlyTargetRules Target) : base(Target)
+		{
+			PublicDependencyModuleNames.AddRange(
+			new string[]
+			{
+				"AnimationCore",
+                "AnimGraphRuntime",
+                "Core",
+                "CoreUObject",
+                "Engine",
+            });
+		}
+	}
+}
diff --git a/Engine/Plugins/Animation/AnimationWarping/Source/Runtime/Private/AnimationWarpingRuntimeModule.cpp b/Engine/Plugins/Animation/AnimationWarping/Source/Runtime/Private/AnimationWarpingRuntimeModule.cpp
new file mode 100644
--- /dev/null
+++ b/Engine/Plugins/Animation/AnimationWarping/Source/Runtime/Private/AnimationWarpingRuntimeModule.cpp
@@ -0,0 +1,74 @@
+// Copyright Epic Games, Inc. All Rights Reserved.
+
+#include "AnimationWarpingRuntimeModule.h"
+#include "Animation/AnimRootMotionProvider.h"
+#include "Animation/AnimNodeBase.h"
+#include "Animation/AnimSequence.h"
+#include "Animation/BuiltInAttributeTypes.h"
+#include "Modules/ModuleManager.h"
+
+class FAnimationWarpingRuntimeModule : public IAnimationWarpingRuntimeModule
+{
+};
+
+namespace UE { namespace AnimationWarping {
+
+// Root motion animation attribute data will always be accessible on the root bone's attribute container
+static const UE::Anim::FAttributeId RootMotionAttributeId = { UE::Anim::IAnimRootMotionProvider::AttributeName , FCompactPoseBoneIndex(0) };
+
+class FModule : public IModuleInterface, public UE::Anim::IAnimRootMotionProvider
+{
+public:
+	virtual void StartupModule() override;
+	virtual void ShutdownModule() override;
+
+public:
+	virtual void SampleRootMotion(const FDeltaTimeRecord& SampleRange, const UAnimSequence& Sequence, bool bLoopingSequence, UE::Anim::FStackAttributeContainer& OutAttributes) const override;
+	virtual bool OverrideRootMotion(const FTransform& RootMotionDelta, UE::Anim::FStackAttributeContainer& OutAttributes) const override;
+	virtual bool ExtractRootMotion(const UE::Anim::FStackAttributeContainer& Attributes, FTransform& OutRootMotionDelta) const override;
+	virtual bool HasRootMotion(const UE::Anim::FStackAttributeContainer& Attributes) const override;
+};
+
+void FModule::StartupModule()
+{
+	IModularFeatures::Get().RegisterModularFeature(UE::Anim::IAnimRootMotionProvider::ModularFeatureName, this);
+}
+
+void FModule::ShutdownModule()
+{
+	IModularFeatures::Get().UnregisterModularFeature(UE::Anim::IAnimRootMotionProvider::ModularFeatureName, this);
+}
+
+void FModule::SampleRootMotion(const FDeltaTimeRecord& SampleRange, const UAnimSequence& Sequence, bool bLoopingSequence, UE::Anim::FStackAttributeContainer& OutAttributes) const
+{
+	const FTransform RootMotionTransform = Sequence.ExtractRootMotion(SampleRange.GetPrevious(), SampleRange.Delta, bLoopingSequence);
+	FTransformAnimationAttribute* RootMotionAttribute = OutAttributes.FindOrAdd<FTransformAnimationAttribute>(RootMotionAttributeId);
+	RootMotionAttribute->Value = RootMotionTransform;
+}
+
+bool FModule::OverrideRootMotion(const FTransform& RootMotionDelta, UE::Anim::FStackAttributeContainer& OutAttributes) const
+{
+	if (FTransformAnimationAttribute* RootMotionAttribute = OutAttributes.Find<FTransformAnimationAttribute>(RootMotionAttributeId))
+	{
+		RootMotionAttribute->Value = RootMotionDelta;
+		return true;
+	}
+
+	return false;
+}
+
+bool FModule::ExtractRootMotion(const UE::Anim::FStackAttributeContainer& Attributes, FTransform& OutRootMotionDelta) const
+{
+	const FTransformAnimationAttribute* RootMotionAttribute = Attributes.Find<FTransformAnimationAttribute>(RootMotionAttributeId);
+	OutRootMotionDelta = RootMotionAttribute ? RootMotionAttribute->Value : FTransform::Identity;
+	return !!RootMotionAttribute;
+}
+
+bool FModule::HasRootMotion(const UE::Anim::FStackAttributeContainer& Attributes) const
+{
+	return !!Attributes.Find<FTransformAnimationAttribute>(RootMotionAttributeId);
+}
+
+}}
+
+IMPLEMENT_MODULE(UE::AnimationWarping::FModule, AnimationWarpingRuntime)
diff --git a/Engine/Plugins/Animation/AnimationWarping/Source/Runtime/Private/BoneControllers/AnimNode_OrientationWarping.cpp b/Engine/Plugins/Animation/AnimationWarping/Source/Runtime/Private/BoneControllers/AnimNode_OrientationWarping.cpp
new file mode 100644
--- /dev/null
+++ b/Engine/Plugins/Animation/AnimationWarping/Source/Runtime/Private/BoneControllers/AnimNode_OrientationWarping.cpp
@@ -0,0 +1,439 @@
+// Copyright Epic Games, Inc. All Rights Reserved.
+
+#include "BoneControllers/AnimNode_OrientationWarping.h"
+#include "Animation/AnimInstanceProxy.h"
+#include "Animation/AnimRootMotionProvider.h"
+
+DECLARE_CYCLE_STAT(TEXT("OrientationWarping Eval"), STAT_OrientationWarping_Eval, STATGROUP_Anim);
+
+#if ENABLE_ANIM_DEBUG
+static TAutoConsoleVariable<int32> CVarAnimNodeOrientationWarpingDebug(TEXT("a.AnimNode.OrientationWarping.Debug"), 0, TEXT("Turn on visualization debugging for Orientation Warping"));
+static TAutoConsoleVariable<int32> CVarAnimNodeOrientationWarpingVerbose(TEXT("a.AnimNode.OrientationWarping.Verbose"), 0, TEXT("Turn on verbose graph debugging for Orientation Warping"));
+static TAutoConsoleVariable<int32> CVarAnimNodeOrientationWarpingEnable(TEXT("a.AnimNode.OrientationWarping.Enable"), 1, TEXT("Toggle Orientation Warping"));
+#endif
+
+namespace UE::Anim
+{
+	static inline FVector GetAxisVector(const EAxis::Type& InAxis)
+	{
+		switch (InAxis)
+		{
+		case EAxis::X:
+			return FVector::ForwardVector;
+		case EAxis::Y:
+			return FVector::RightVector;
+		default:
+			return FVector::UpVector;
+		};
+	}
+
+	static inline bool IsInvalidWarpingAngleDegrees(float Angle, float Tolerance)
+	{
+		Angle = FRotator::NormalizeAxis(Angle);
+		return FMath::IsNearlyZero(Angle, Tolerance) || FMath::IsNearlyEqual(FMath::Abs(Angle), 180.f, Tolerance);
+	}
+}
+
+void FAnimNode_OrientationWarping::GatherDebugData(FNodeDebugData& DebugData)
+{
+	FString DebugLine = DebugData.GetNodeName(this);
+#if ENABLE_ANIM_DEBUG
+	if (CVarAnimNodeOrientationWarpingVerbose.GetValueOnAnyThread() == 1)
+	{
+		if (Mode == EWarpingEvaluationMode::Manual)
+		{
+			DebugLine += TEXT("\n - Evaluation Mode: (Manual)");
+			DebugLine += FString::Printf(TEXT("\n - Orientation Angle: (%.3fd)"), FMath::RadiansToDegrees(ActualOrientationAngle));
+		}
+		else
+		{
+			DebugLine += TEXT("\n - Evaluation Mode: (Graph)");
+			DebugLine += FString::Printf(TEXT("\n - Orientation Angle: (%.3fd)"), FMath::RadiansToDegrees(ActualOrientationAngle));
+			DebugLine += FString::Printf(TEXT("\n - Locomotion Angle: (%.3fd)"), FMath::RadiansToDegrees(LocomotionAngle));
+			DebugLine += FString::Printf(TEXT("\n - Locomotion Delta Angle Threshold: (%.3fd)"), LocomotionAngleDeltaThreshold);
+		}
+		DebugLine += FString::Printf(TEXT("\n - Distributed Bone Orientation Alpha: (%.3fd)"), DistributedBoneOrientationAlpha);
+		if (const UEnum* TypeEnum = FindObject<UEnum>(ANY_PACKAGE, TEXT("EAxis")))
+		{
+			DebugLine += FString::Printf(TEXT("\n - Rotation Axis: (%s)"), *(TypeEnum->GetNameStringByIndex(static_cast<int32>(RotationAxis))));
+		}
+		DebugLine += FString::Printf(TEXT("\n - Rotation Interpolation Speed: (%.3fd)"), RotationInterpSpeed);
+	}
+	else
+#endif
+	{
+	DebugLine += FString::Printf(TEXT("(Orientation Angle: %.3fd)"), FMath::RadiansToDegrees(ActualOrientationAngle));
+	}
+	DebugData.AddDebugItem(DebugLine);
+	ComponentPose.GatherDebugData(DebugData);
+}
+
+void FAnimNode_OrientationWarping::Initialize_AnyThread(const FAnimationInitializeContext& Context)
+{
+	FAnimNode_SkeletalControlBase::Initialize_AnyThread(Context);
+
+	PreviousRootMotionDeltaDirection = FVector::ZeroVector;
+	PreviousOrientationAngle = 0.f;
+	ActualOrientationAngle = 0.f;
+}
+
+void FAnimNode_OrientationWarping::UpdateInternal(const FAnimationUpdateContext& Context)
+{
+	FAnimNode_SkeletalControlBase::UpdateInternal(Context);
+}
+
+void FAnimNode_OrientationWarping::EvaluateSkeletalControl_AnyThread(FComponentSpacePoseContext& Output, TArray<FBoneTransform>& OutBoneTransforms)
+{
+	SCOPE_CYCLE_COUNTER(STAT_OrientationWarping_Eval);
+	check(OutBoneTransforms.Num() == 0);
+
+	ActualOrientationAngle = OrientationAngle;
+
+	const FVector RotationAxisVector = UE::Anim::GetAxisVector(RotationAxis);
+	FVector RootMotionDeltaDirection = FVector::ZeroVector;
+	FVector LocomotionForward = FVector::ZeroVector;
+
+	bool bGraphDrivenWarping = false;
+	const UE::Anim::IAnimRootMotionProvider* RootMotionProvider = UE::Anim::IAnimRootMotionProvider::Get();
+
+	if (Mode == EWarpingEvaluationMode::Graph)
+	{
+		bGraphDrivenWarping = !!RootMotionProvider;
+		ensureMsgf(bGraphDrivenWarping, TEXT("Graph driven Orientation Warping expected a valid root motion delta provider interface."));
+	}
+
+	// We will likely need to revisit LocomotionAngle participating as an input to orientation warping.
+	// Without velocity information from the motion model (such as the capsule), LocomotionAngle isn't enough
+	// information in isolation for all cases when deciding to warp.
+	//
+	// For example imagine that the motion model has stopped moving with zero velocity due to a
+	// transition into a strafing stop. During that transition we may play an animation with non-zero 
+	// velocity for an arbitrary number of frames. In this scenario the concept of direction is meaningless 
+	// since we cannot orient the animation to match a zero velocity and consequently a zero direction, 
+	// since that would break the pose. For those frames, we would incorrectly over-orient the strafe.
+	//
+	// The solution may be instead to pass velocity with the actor base rotation, allowing us to retain
+	// speed information about the motion. It may also allow us to do more complex orienting behavior 
+	// when multiple degrees of freedom can be considered.
+
+	if (bGraphDrivenWarping)
+	{
+		FTransform RootMotionTransformDelta = FTransform::Identity;
+		bGraphDrivenWarping = RootMotionProvider->ExtractRootMotion(Output.CustomAttributes, RootMotionTransformDelta);
+
+		// Graph driven orientation warping will modify the incoming root motion to orient towards the intended locomotion angle
+		if (ensureMsgf(bGraphDrivenWarping, TEXT(
+			"Graph driven Orientation Warping expected a root motion delta to be present in the attribute stream.\n"
+			"Check that the evaluating Animation Sequence(s) have: \"Enable Root Motion\" enabled.")))
+		{
+			// In UE, forward is defined as +x; consequently this is also true when sampling an actor's velocity. Historically the skeletal 
+			// mesh component forward will not match the actor, requiring us to correct the rotation before sampling the LocomotionForward.
+			// In order to make orientation warping 'pure' in the future we will need to provide more context about the intent of
+			// the actor vs the intent of the animation in their respective spaces. Specifically, we will need some form the following information:
+			//
+			// 1. Actor Forward
+			// 2. Actor Velocity
+			// 3. Skeletal Mesh Relative Rotation
+
+			LocomotionAngle = FRotator::NormalizeAxis(LocomotionAngle);
+			LocomotionAngle = FMath::DegreesToRadians(LocomotionAngle);
+			const FQuat LocomotionRotation = FQuat(RotationAxisVector, LocomotionAngle);
+
+			const FTransform SkeletalMeshRelativeTransform = Output.AnimInstanceProxy->GetComponentRelativeTransform();
+			const FQuat SkeletalMeshRelativeRotation = SkeletalMeshRelativeTransform.GetRotation();
+			LocomotionForward = SkeletalMeshRelativeRotation.UnrotateVector(LocomotionRotation.GetForwardVector()).GetSafeNormal();
+
+			const FVector RootMotionDeltaTranslation = RootMotionTransformDelta.GetTranslation();
+			RootMotionDeltaDirection = RootMotionDeltaTranslation.GetSafeNormal();
+
+			// Capture the delta rotation from the axis of motion we care about
+			FQuat WarpedRotation = FQuat::FindBetween(RootMotionDeltaDirection, LocomotionForward);
+			ActualOrientationAngle = WarpedRotation.GetTwistAngle(RotationAxisVector);
+
+			// Motion Matching may return an animation that deviates a lot from the movement direction (e.g movement direction going bwd and motion matching could return the fwd animation for a few frames)
+			// When that happens, since we use the delta between root motion and movement direction, we would be over-rotating the lower body and breaking the pose during those frames
+			// So, when that happens we use the inverse of the movement direction to calculate our target rotation. 
+			// This feels a bit 'hacky' but its the only option I've found so far to mitigate the problem
+			if (LocomotionAngleDeltaThreshold > 0.f && FMath::Abs(FMath::RadiansToDegrees(ActualOrientationAngle)) > LocomotionAngleDeltaThreshold)
+			{
+				WarpedRotation = FQuat::FindBetween(RootMotionDeltaDirection, -LocomotionForward);
+				ActualOrientationAngle = WarpedRotation.GetTwistAngle(RotationAxisVector);
+			}
+
+			// For interpolated warping, guarantee that PreviousOrientationAngle is respect to the current frame's root motion direction 
+			const FVector RootMotionCrossProduct = RootMotionDeltaDirection.Cross(PreviousRootMotionDeltaDirection);
+			float RootMotionDeltaAngle = FMath::Acos(RootMotionDeltaDirection.Dot(PreviousRootMotionDeltaDirection));
+			RootMotionDeltaAngle *= FMath::Sign(RotationAxisVector.Dot(RootMotionCrossProduct));
+
+			PreviousRootMotionDeltaDirection = RootMotionDeltaDirection;
+			PreviousOrientationAngle += RootMotionDeltaAngle;
+
+			// Rotate the root motion delta fully by the warped angle
+			const FVector WarpedRootMotionTranslationDelta = WarpedRotation.RotateVector(RootMotionDeltaTranslation);
+			RootMotionTransformDelta.SetTranslation(WarpedRootMotionTranslationDelta);
+
+			// Forward the side effects of orientation warping on the root motion contribution for this sub-graph
+			const bool bRootMotionOverridden = RootMotionProvider->OverrideRootMotion(RootMotionTransformDelta, Output.CustomAttributes);
+			ensureMsgf(bRootMotionOverridden, TEXT("Graph driven Orientation Warping expected a root motion delta to be present in the attribute stream prior to warping/overriding it."));
+		}
+		else
+		{
+			// Early exit on broken graph driven behavior 
+			return;
+		}
+	} 
+	else
+	{
+		// Manual orientation warping will take the angle directly
+		ActualOrientationAngle = FRotator::NormalizeAxis(ActualOrientationAngle);
+		ActualOrientationAngle = FMath::DegreesToRadians(ActualOrientationAngle);
+	}
+
+	// Optionally interpolate the effective orientation towards the target orientation angle
+	if (RotationInterpSpeed > 0.f)
+	{
+		ActualOrientationAngle = FMath::FInterpTo(PreviousOrientationAngle, ActualOrientationAngle, Output.AnimInstanceProxy->GetDeltaSeconds(), RotationInterpSpeed);
+		PreviousOrientationAngle = ActualOrientationAngle;
+	}
+
+	// Allow the alpha value of the node to affect the final rotation
+	ActualOrientationAngle *= ActualAlpha;
+
+#if ENABLE_ANIM_DEBUG
+	bool bDebugging = false;
+#if WITH_EDITORONLY_DATA
+	bDebugging = bDebugging || bEnableDebugDraw;
+#else
+	constexpr float DebugDrawScale = 1.f;
+#endif
+	bDebugging = bDebugging || CVarAnimNodeOrientationWarpingDebug.GetValueOnAnyThread() == 1;
+
+	if (bDebugging)
+	{
+		const FTransform ComponentTransform = Output.AnimInstanceProxy->GetComponentTransform();
+		const FVector ActorForwardDirection = Output.AnimInstanceProxy->GetActorTransform().GetRotation().GetForwardVector();
+
+		const FVector ForwardDirection = bGraphDrivenWarping
+			? ComponentTransform.GetRotation().RotateVector(LocomotionForward)
+			: ActorForwardDirection;
+
+		FVector DebugArrowOffset = FVector::ZAxisVector * DebugDrawScale;
+		Output.AnimInstanceProxy->AnimDrawDebugDirectionalArrow(
+			ComponentTransform.GetLocation() + DebugArrowOffset,
+			ComponentTransform.GetLocation() + DebugArrowOffset + ForwardDirection * 100.f * DebugDrawScale,
+			40.f * DebugDrawScale, FColor::Red, false, 0.f, 2.f * DebugDrawScale);
+
+		const FVector RotationDirection = bGraphDrivenWarping
+			? ComponentTransform.GetRotation().RotateVector(RootMotionDeltaDirection)
+			: ActorForwardDirection.RotateAngleAxis(OrientationAngle, RotationAxisVector);
+
+		DebugArrowOffset += DebugArrowOffset;
+		Output.AnimInstanceProxy->AnimDrawDebugDirectionalArrow(
+			ComponentTransform.GetLocation() + DebugArrowOffset,
+			ComponentTransform.GetLocation() + DebugArrowOffset + RotationDirection * 100.f * DebugDrawScale,
+			40.f * DebugDrawScale, FColor::Blue, false, 0.f, 2.f * DebugDrawScale);
+
+		const float ActualOrientationAngleDegrees = FMath::RadiansToDegrees(ActualOrientationAngle);
+		const FVector WarpedRotationDirection = bGraphDrivenWarping 
+			? RotationDirection.RotateAngleAxis(ActualOrientationAngleDegrees, RotationAxisVector)
+			: ActorForwardDirection.RotateAngleAxis(ActualOrientationAngleDegrees, RotationAxisVector);
+
+		DebugArrowOffset += DebugArrowOffset;
+		Output.AnimInstanceProxy->AnimDrawDebugDirectionalArrow(
+			ComponentTransform.GetLocation() + DebugArrowOffset,
+			ComponentTransform.GetLocation() + DebugArrowOffset + WarpedRotationDirection * 100.f * DebugDrawScale,
+			40.f * DebugDrawScale, FColor::Green, false, 0.f, 2.f * DebugDrawScale);
+	}
+#endif
+		
+	// Rotate Root Bone first, as that cheaply rotates the whole pose with one transformation.
+	if (!FMath::IsNearlyZero(DistributedBoneOrientationAlpha, KINDA_SMALL_NUMBER))
+	{
+		const FQuat RootRotation = FQuat(RotationAxisVector, ActualOrientationAngle * DistributedBoneOrientationAlpha);
+		const FCompactPoseBoneIndex RootBoneIndex(0);
+
+		FTransform RootBoneTransform(Output.Pose.GetComponentSpaceTransform(RootBoneIndex));
+		RootBoneTransform.SetRotation(RootRotation * RootBoneTransform.GetRotation());
+		RootBoneTransform.NormalizeRotation();
+		Output.Pose.SetComponentSpaceTransform(RootBoneIndex, RootBoneTransform);
+	}
+
+	const int32 NumSpineBones = SpineBoneDataArray.Num();
+	const bool bSpineOrientationAlpha = !FMath::IsNearlyZero(DistributedBoneOrientationAlpha, KINDA_SMALL_NUMBER);
+	const bool bUpdateSpineBones = (NumSpineBones > 0) && bSpineOrientationAlpha;
+
+	if (bUpdateSpineBones)
+	{
+		// Spine bones counter rotate body orientation evenly across all bones.
+		for (int32 ArrayIndex = 0; ArrayIndex < NumSpineBones; ArrayIndex++)
+		{
+			const FOrientationWarpingSpineBoneData& BoneData = SpineBoneDataArray[ArrayIndex];
+			const FQuat SpineBoneCounterRotation = FQuat(RotationAxisVector, -ActualOrientationAngle * DistributedBoneOrientationAlpha * BoneData.Weight);
+			check(BoneData.Weight > 0.f);
+
+			FTransform SpineBoneTransform(Output.Pose.GetComponentSpaceTransform(BoneData.BoneIndex));
+			SpineBoneTransform.SetRotation((SpineBoneCounterRotation * SpineBoneTransform.GetRotation()));
+			SpineBoneTransform.NormalizeRotation();
+			Output.Pose.SetComponentSpaceTransform(BoneData.BoneIndex, SpineBoneTransform);
+		}
+	}
+
+	const float IKFootRootOrientationAlpha = 1.f - DistributedBoneOrientationAlpha;
+	const bool bUpdateIKFootRoot = (IKFootData.IKFootRootBoneIndex != FCompactPoseBoneIndex(INDEX_NONE)) && !FMath::IsNearlyZero(IKFootRootOrientationAlpha, KINDA_SMALL_NUMBER);
+
+	// Rotate IK Foot Root
+	if (bUpdateIKFootRoot)
+	{
+		const FQuat BoneRotation = FQuat(RotationAxisVector, ActualOrientationAngle * IKFootRootOrientationAlpha);
+
+		FTransform IKFootRootTransform(Output.Pose.GetComponentSpaceTransform(IKFootData.IKFootRootBoneIndex));
+		IKFootRootTransform.SetRotation(BoneRotation * IKFootRootTransform.GetRotation());
+		IKFootRootTransform.NormalizeRotation();
+		Output.Pose.SetComponentSpaceTransform(IKFootData.IKFootRootBoneIndex, IKFootRootTransform);
+
+		// IK Feet 
+		// These match the root orientation, so don't rotate them. Just preserve root rotation. 
+		// We need to update their translation though, since we rotated their parent (the IK Foot Root bone).
+		const int32 NumIKFootBones = IKFootData.IKFootBoneIndexArray.Num();
+		const bool bUpdateIKFootBones = bUpdateIKFootRoot && (NumIKFootBones > 0);
+
+		if (bUpdateIKFootBones)
+		{
+			const FQuat IKFootRotation = FQuat(RotationAxisVector, -ActualOrientationAngle * IKFootRootOrientationAlpha);
+
+			for (int32 ArrayIndex = 0; ArrayIndex < NumIKFootBones; ArrayIndex++)
+			{
+				const FCompactPoseBoneIndex& IKFootBoneIndex = IKFootData.IKFootBoneIndexArray[ArrayIndex];
+
+				FTransform IKFootBoneTransform(Output.Pose.GetComponentSpaceTransform(IKFootBoneIndex));
+				IKFootBoneTransform.SetRotation(IKFootRotation * IKFootBoneTransform.GetRotation());
+				IKFootBoneTransform.NormalizeRotation();
+				Output.Pose.SetComponentSpaceTransform(IKFootBoneIndex, IKFootBoneTransform);
+			}
+		}
+	}
+
+	OutBoneTransforms.Sort(FCompareBoneTransformIndex());
+}
+
+bool FAnimNode_OrientationWarping::IsValidToEvaluate(const USkeleton* Skeleton, const FBoneContainer& RequiredBones)
+{
+#if ENABLE_ANIM_DEBUG
+	if (CVarAnimNodeOrientationWarpingEnable.GetValueOnAnyThread() == 0)
+	{
+		return false;
+	}
+#endif
+	if (RotationAxis == EAxis::None)
+	{
+		return false;
+	}
+
+	if (Mode == EWarpingEvaluationMode::Manual && UE::Anim::IsInvalidWarpingAngleDegrees(OrientationAngle, KINDA_SMALL_NUMBER))
+	{
+		return false;
+	}
+
+	if (SpineBoneDataArray.IsEmpty())
+	{
+		return false;
+	}
+	else
+	{
+		for (const auto& Spine : SpineBoneDataArray)
+		{
+			if (Spine.BoneIndex == INDEX_NONE)
+			{
+				return false;
+			}
+		}
+	}
+
+	if (IKFootData.IKFootRootBoneIndex == INDEX_NONE)
+	{
+		return false;
+	}
+
+	if (IKFootData.IKFootBoneIndexArray.IsEmpty())
+	{
+		return false;
+	}
+	else
+	{
+		for (const auto& IKFootBoneIndex : IKFootData.IKFootBoneIndexArray)
+		{
+			if (IKFootBoneIndex == INDEX_NONE)
+			{
+				return false;
+			}
+		}
+	}
+	return true;
+}
+
+void FAnimNode_OrientationWarping::InitializeBoneReferences(const FBoneContainer& RequiredBones)
+{
+	IKFootRootBone.Initialize(RequiredBones);
+	IKFootData.IKFootRootBoneIndex = IKFootRootBone.GetCompactPoseIndex(RequiredBones);
+
+	IKFootData.IKFootBoneIndexArray.Reset();
+	for (auto& Bone : IKFootBones)
+	{
+		Bone.Initialize(RequiredBones);
+		IKFootData.IKFootBoneIndexArray.Add(Bone.GetCompactPoseIndex(RequiredBones));
+	}
+
+	SpineBoneDataArray.Reset();
+	for (auto& Bone : SpineBones)
+	{
+		Bone.Initialize(RequiredBones);
+		SpineBoneDataArray.Add(FOrientationWarpingSpineBoneData(Bone.GetCompactPoseIndex(RequiredBones)));
+	}
+
+	if (SpineBoneDataArray.Num() > 0)
+	{
+		// Sort bones indices so we can transform parent before child
+		SpineBoneDataArray.Sort(FOrientationWarpingSpineBoneData::FCompareBoneIndex());
+
+		// Assign Weights.
+		TArray<int32, TInlineAllocator<20>> IndicesToUpdate;
+
+		for (int32 Index = SpineBoneDataArray.Num() - 1; Index >= 0; Index--)
+		{
+			// If this bone's weight hasn't been updated, scan his parents.
+			// If parents have weight, we add it to 'ExistingWeight'.
+			// split (1.f - 'ExistingWeight') between all members of the chain that have no weight yet.
+			if (SpineBoneDataArray[Index].Weight == 0.f)
+			{
+				IndicesToUpdate.Reset(SpineBoneDataArray.Num());
+				float ExistingWeight = 0.f;
+				IndicesToUpdate.Add(Index);
+
+				const FCompactPoseBoneIndex CompactBoneIndex = SpineBoneDataArray[Index].BoneIndex;
+				for (int32 ParentIndex = Index - 1; ParentIndex >= 0; ParentIndex--)
+				{
+					if (RequiredBones.BoneIsChildOf(CompactBoneIndex, SpineBoneDataArray[ParentIndex].BoneIndex))
+					{
+						if (SpineBoneDataArray[ParentIndex].Weight > 0.f)
+						{
+							ExistingWeight += SpineBoneDataArray[ParentIndex].Weight;
+						}
+						else
+						{
+							IndicesToUpdate.Add(ParentIndex);
+						}
+					}
+				}
+
+				check(IndicesToUpdate.Num() > 0);
+				const float WeightToShare = 1.f - ExistingWeight;
+				const float IndividualWeight = WeightToShare / float(IndicesToUpdate.Num());
+
+				for (int32 UpdateListIndex = 0; UpdateListIndex < IndicesToUpdate.Num(); UpdateListIndex++)
+				{
+					SpineBoneDataArray[IndicesToUpdate[UpdateListIndex]].Weight = IndividualWeight;
+				}
+			}
+		}
+	}
+}
\ No newline at end of file
diff --git a/Engine/Plugins/Animation/AnimationWarping/Source/Runtime/Private/BoneControllers/AnimNode_SlopeWarping.cpp b/Engine/Plugins/Animation/AnimationWarping/Source/Runtime/Private/BoneControllers/AnimNode_SlopeWarping.cpp
new file mode 100644
--- /dev/null
+++ b/Engine/Plugins/Animation/AnimationWarping/Source/Runtime/Private/BoneControllers/AnimNode_SlopeWarping.cpp
@@ -0,0 +1,729 @@
+// Copyright Epic Games, Inc. All Rights Reserved.
+
+#include "BoneControllers/AnimNode_SlopeWarping.h"
+#include "AnimationRuntime.h"
+#include "Animation/AnimInstanceProxy.h"
+#include "GameFramework/Character.h"
+#include "GameFramework/CharacterMovementComponent.h"
+#include "Components/CapsuleComponent.h"
+#include "Engine/SkeletalMeshSocket.h"
+#include "DrawDebugHelpers.h"
+
+#if ENABLE_ANIM_DEBUG
+TAutoConsoleVariable<int32> CVarAnimNodeSlopeWarpingDebug(TEXT("a.AnimNode.SlopeWarping.Debug"), 0, TEXT("Turn on debug for AnimNode_SlopeWarping"));
+#endif
+
+TAutoConsoleVariable<int32> CVarSlopeWarpingEnable(TEXT("a.AnimNode.SlopeWarping.Enable"), 1, TEXT("Toggle Slope Warping"));
+
+DECLARE_CYCLE_STAT(TEXT("Slope Warping Eval"), STAT_SlopeWarping_Eval, STATGROUP_Anim);
+
+FAnimNode_SlopeWarping::FAnimNode_SlopeWarping()
+	: MySkelMeshComponent(nullptr)
+	, MyMovementComponent(nullptr)
+	, MyAnimInstanceProxy(nullptr)
+	, GravityDir(-FVector::UpVector)
+	, CustomFloorOffset(FVector::ZeroVector)
+	, CachedDeltaTime(0.f)
+	, TargetFloorNormalWorldSpace(ForceInitToZero)
+	, TargetFloorOffsetLocalSpace(ForceInitToZero)
+	, MaxStepHeight(50.f)
+	, bKeepMeshInsideOfCapsule(true)
+	, bPullPelvisDown(true)
+	, bUseCustomFloorOffset(false)
+	, bWasOnGround(false)
+	, bShowDebug(false)
+	, bFloorSmoothingInitialized(false)
+	, ActorLocation(ForceInitToZero)
+	, GravityDirCompSpace(ForceInitToZero)
+{
+	// Defaults
+	PelvisOffsetInterpolator.SetDefaultSpringConstants(250.f);
+	FloorNormalInterpolator.SetDefaultSpringConstants(1000.f);
+	FloorOffsetInterpolator.SetDefaultSpringConstants(1000.f);
+}
+
+void FAnimNode_SlopeWarping::GatherDebugData(FNodeDebugData& DebugData)
+{
+	FString DebugLine = DebugData.GetNodeName(this);
+
+	// 	DebugLine += "(";
+	// 	AddDebugNodeData(DebugLine);
+	// 	DebugLine += FString::Printf(TEXT(" Target: %s)"), *BoneToModify.BoneName.ToString());
+
+	DebugData.AddDebugItem(DebugLine);
+
+	ComponentPose.GatherDebugData(DebugData);
+}
+
+void FAnimNode_SlopeWarping::Initialize_AnyThread(const FAnimationInitializeContext& Context)
+{
+	FAnimNode_SkeletalControlBase::Initialize_AnyThread(Context);
+
+	MyAnimInstanceProxy = Context.AnimInstanceProxy;
+	MySkelMeshComponent = Context.AnimInstanceProxy ? Context.AnimInstanceProxy->GetSkelMeshComponent() : nullptr;
+	ACharacter* CharacterOwner = MySkelMeshComponent ? Cast<ACharacter>(MySkelMeshComponent->GetOwner()) : nullptr;
+	MyMovementComponent = CharacterOwner ? CharacterOwner->GetCharacterMovement() : nullptr;
+
+	FloorNormalInterpolator.SetPosition(-GravityDir);
+}
+
+void FAnimNode_SlopeWarping::UpdateInternal(const FAnimationUpdateContext& Context)
+{
+	FAnimNode_SkeletalControlBase::UpdateInternal(Context);
+
+	CachedDeltaTime += Context.GetDeltaTime();
+}
+
+static FVector GetFootIKCentroid(const TArray<FSlopeWarpingFootData>& FeetData)
+{
+	check(FeetData.Num() > 0);
+
+	FVector Centroid = FVector::ZeroVector;
+	for (auto& FootData : FeetData)
+	{
+		Centroid += FootData.IKFootTransform.GetLocation();
+	}
+
+	Centroid /= float(FeetData.Num());
+	return Centroid;
+}
+
+static FVector ReOrientNormal(const FVector& InNormal, FVector& PointA, const FVector& PointB)
+{
+	const FVector AxisX = (PointA - PointB).GetSafeNormal();
+	if (!AxisX.IsNearlyZero() && !InNormal.IsNearlyZero() && (FMath::Abs(AxisX | InNormal) > DELTA))
+	{
+		const FVector AxisY = AxisX ^ InNormal;
+		const FVector AxisZ = AxisX ^ AxisY;
+
+		// Make sure our normal points upwards. (take into account gravity dir?)
+		return ((AxisZ | FVector::UpVector) > 0.f) ? AxisZ : -AxisZ;
+	}
+	else
+	{
+		return InNormal;
+	}
+}
+
+FVector FAnimNode_SlopeWarping::GetFloorNormalFromMovementComponent() const
+{
+	check(MyMovementComponent);
+	return (MyMovementComponent->CurrentFloor.bBlockingHit) ? MyMovementComponent->CurrentFloor.HitResult.ImpactNormal : -GravityDir;
+}
+
+void FAnimNode_SlopeWarping::SetWorldSpaceTargetFloorInfoFromMovementComponent()
+{
+	TargetFloorNormalWorldSpace = GetFloorNormalFromMovementComponent();
+
+	check(MySkelMeshComponent);
+	check(MyMovementComponent);
+
+	FVector WorldSpaceFloorOffset(FVector::ZeroVector);
+
+	// Take into account offsets created by capsule.
+	if (ACharacter* CharacterOwner = Cast<ACharacter>(MySkelMeshComponent->GetOwner()))
+	{
+		UCapsuleComponent* Capsule = Cast<UCapsuleComponent>(CharacterOwner->GetRootComponent());
+		if (Capsule)
+		{
+			// Take into account MovementComponent floor dist offset
+			const FVector MovementComponentFloorDistOffset = MyMovementComponent->CurrentFloor.bBlockingHit ? MyMovementComponent->CurrentFloor.GetDistanceToFloor() * GravityDir : FVector::ZeroVector;
+			WorldSpaceFloorOffset += MovementComponentFloorDistOffset;
+
+			// Take into account offset between mesh and base of capsule
+			FVector MeshToCapsuleBaseOffset = FVector::ZeroVector;
+			if (MyMovementComponent->CurrentFloor.bBlockingHit)
+			{
+				const FVector CapsuleBaseLoc = Capsule->GetComponentLocation() + GravityDir * Capsule->GetScaledCapsuleHalfHeight();
+				const FVector MeshLoc = MySkelMeshComponent->GetComponentLocation();
+				const FVector ProjectedMeshLoc = (FMath::Abs(GravityDir | TargetFloorNormalWorldSpace) > DELTA) ? FMath::LinePlaneIntersection(MeshLoc, MeshLoc + GravityDir, CapsuleBaseLoc, TargetFloorNormalWorldSpace) : MeshLoc;
+
+				MeshToCapsuleBaseOffset = (MeshLoc - ProjectedMeshLoc);
+				WorldSpaceFloorOffset += MeshToCapsuleBaseOffset;
+			}
+
+			// Take into account Capsule radius against slope.
+			const float CapsuleRadius = Capsule->GetScaledCapsuleRadius();
+			const FVector CapsuleRadiusAgainstSlopeOffset = GravityDir * CapsuleRadius * (1.f - (TargetFloorNormalWorldSpace | -GravityDir));
+			WorldSpaceFloorOffset += CapsuleRadiusAgainstSlopeOffset;
+
+#if ENABLE_ANIM_DEBUG
+			if (bShowDebug)
+			{
+				FString DebugString = FString::Printf(TEXT("DistToFloorOffset: %s, MeshToCapsuleBaseOffset: %s, CapsuleRadiusAgainstSlopeOffset: %s, WorldSpaceFloorOffset: %s"),
+					*MovementComponentFloorDistOffset.ToCompactString(), *MeshToCapsuleBaseOffset.ToCompactString(), *CapsuleRadiusAgainstSlopeOffset.ToCompactString(), *WorldSpaceFloorOffset.ToCompactString());
+				MyAnimInstanceProxy->AnimDrawDebugOnScreenMessage(DebugString, FColor::Red);
+			}
+#endif
+		}
+	}
+
+	// Turn floor offset to local space.
+	TargetFloorOffsetLocalSpace = MySkelMeshComponent->GetComponentTransform().InverseTransformVectorNoScale(WorldSpaceFloorOffset);
+}
+
+static bool HasCharacterTeleported(const FVector& InPreviousLoc, const FVector& CurrentLoc, const UCharacterMovementComponent& InCharacterMovementComponent)
+{
+	return (CurrentLoc - InPreviousLoc).SizeSquared2D() > FMath::Square(InCharacterMovementComponent.NetworkNoSmoothUpdateDistance);
+}
+
+void FAnimNode_SlopeWarping::GetSmoothedFloorInfo(FCSPose<FCompactPose>& MeshBases, FVector& OutFloorNormal, FVector& OutFloorOffset)
+{
+	check(MyMovementComponent);
+	check(MySkelMeshComponent);
+
+	AActor* OwnerActor = MySkelMeshComponent->GetOwner();
+	// Use mesh component location instead to help simulated proxies where capsule simulation is disabled, and it's mesh location is interpolated instead.
+	const FVector NewActorLocation = MySkelMeshComponent->GetComponentLocation();
+
+	// Detect when Actor teleports to reset our interpolation
+	const bool bJustTeleported = !bFloorSmoothingInitialized || HasCharacterTeleported(ActorLocation, NewActorLocation, *MyMovementComponent);
+	const bool bIsOnGround = ((MyMovementComponent->MovementMode == MOVE_Walking) || (MyMovementComponent->MovementMode == MOVE_NavWalking)) && MyMovementComponent->CurrentFloor.bBlockingHit;
+
+	const FVector ActorLocationLastFrame = ActorLocation;
+	ActorLocation = NewActorLocation;
+
+	// if bFloorSmoothingInitialized is false, then bJustTeleported should be true
+	check(bFloorSmoothingInitialized || (bFloorSmoothingInitialized != bJustTeleported));
+	if (bJustTeleported)
+	{
+		bFloorSmoothingInitialized = true;
+
+		TargetFloorOffsetLocalSpace = FVector::ZeroVector;
+		FloorOffsetInterpolator.Reset(TargetFloorOffsetLocalSpace);
+
+		TargetFloorNormalWorldSpace = -GravityDir;
+		FloorNormalInterpolator.Reset(TargetFloorNormalWorldSpace);
+	}
+	else
+	{
+		if (bIsOnGround)
+		{
+			// Compensate for sudden capsule moves
+			if (bWasOnGround)
+			{
+				const FVector CapsuleFloorNormal = GetFloorNormalFromMovementComponent();
+				const FVector AdjustedActorLocationLastFrame = (FMath::Abs(GravityDir | CapsuleFloorNormal) > DELTA) ? FMath::LinePlaneIntersection(ActorLocation, ActorLocation + GravityDir, ActorLocationLastFrame, CapsuleFloorNormal) : ActorLocation;
+
+				// Delta Capsule
+				const FVector CapsuleMoveOffsetWorldSpace = (ActorLocation - AdjustedActorLocationLastFrame);
+				if (!CapsuleMoveOffsetWorldSpace.IsNearlyZero(KINDA_SMALL_NUMBER))
+				{
+					const FVector CapsuleMoveOffsetLocalSpace = MySkelMeshComponent->GetComponentTransform().InverseTransformVectorNoScale(CapsuleMoveOffsetWorldSpace);
+					FloorOffsetInterpolator.OffsetPosition(-CapsuleMoveOffsetLocalSpace);
+					TargetFloorOffsetLocalSpace -= CapsuleMoveOffsetLocalSpace;
+
+#if ENABLE_ANIM_DEBUG
+					if (bShowDebug)
+					{
+						FString DebugString = FString::Printf(TEXT("CapsuleMoveOffset: %s"), *CapsuleMoveOffsetLocalSpace.ToCompactString());
+						MyAnimInstanceProxy->AnimDrawDebugOnScreenMessage(DebugString, FColor::Red);
+					}
+#endif
+				}
+			}
+
+			SetWorldSpaceTargetFloorInfoFromMovementComponent();
+		}
+		else
+		{
+#if ENABLE_ANIM_DEBUG
+			if (bShowDebug)
+			{
+				FString DebugString = FString::Printf(TEXT("SlopeWarping NOT ON GROUND"));
+				MyAnimInstanceProxy->AnimDrawDebugOnScreenMessage(DebugString, FColor::Red);
+			}
+#endif
+			TargetFloorOffsetLocalSpace = FVector::ZeroVector;
+			TargetFloorNormalWorldSpace = -GravityDir;
+		}
+	}
+
+	bWasOnGround = bIsOnGround;
+
+	// Interpolate floor normal in world space.
+	if (!FloorNormalInterpolator.IsPositionEqualTo(TargetFloorNormalWorldSpace))
+	{
+		const FVector SmoothedFloorNormalWorldSpace = FloorNormalInterpolator.Update(TargetFloorNormalWorldSpace, CachedDeltaTime).GetSafeNormal();
+		FloorNormalInterpolator.SetPosition(SmoothedFloorNormalWorldSpace);
+	}
+	OutFloorNormal = MySkelMeshComponent->GetComponentTransform().InverseTransformVectorNoScale(FloorNormalInterpolator.GetPosition());
+
+	if (bUseCustomFloorOffset)
+	{
+		OutFloorOffset = CustomFloorOffset;
+		FloorOffsetInterpolator.SetPosition(CustomFloorOffset);
+	}
+	else
+	{
+		// Interpolate floor offset in local space.
+		const FVector MaxStep = -GravityDir * MaxStepHeight;
+		TargetFloorOffsetLocalSpace = ClampVector(TargetFloorOffsetLocalSpace, -MaxStep, MaxStep);
+		OutFloorOffset = FloorOffsetInterpolator.Update(TargetFloorOffsetLocalSpace, CachedDeltaTime);
+
+#if ENABLE_ANIM_DEBUG
+		if (bShowDebug)
+		{
+			FString DebugString = FString::Printf(TEXT("OutFloorOffset: %s, TargetFloorOffsetLocalSpace: %s")
+				, *OutFloorOffset.ToCompactString(), *TargetFloorOffsetLocalSpace.ToCompactString());
+			MyAnimInstanceProxy->AnimDrawDebugOnScreenMessage(DebugString, FColor::Red);
+		}
+#endif
+	}
+}
+
+void FAnimNode_SlopeWarping::EvaluateSkeletalControl_AnyThread(FComponentSpacePoseContext& Output, TArray<FBoneTransform>& OutBoneTransforms)
+{
+	SCOPE_CYCLE_COUNTER(STAT_SlopeWarping_Eval);
+
+	check(OutBoneTransforms.Num() == 0);
+
+	const FBoneContainer& RequiredBones = Output.Pose.GetPose().GetBoneContainer();
+
+#if ENABLE_ANIM_DEBUG
+	bShowDebug = (CVarAnimNodeSlopeWarpingDebug.GetValueOnAnyThread() == 1);
+#endif
+
+	const FCompactPoseBoneIndex IKFootRootBoneIndex = IKFootRootBone.GetCompactPoseIndex(RequiredBones);
+	FTransform IKFootRootTransform = Output.Pose.GetComponentSpaceTransform(IKFootRootBoneIndex);
+	if (!ensureMsgf(IKFootRootTransform.ContainsNaN() == false, TEXT("FAnimNode_SlopeWarping, Incoming IKFootRootTransform contains NaN (%s)"), *IKFootRootTransform.ToString()))
+	{
+		IKFootRootTransform = FTransform::Identity;
+	}
+
+	const FVector IKFootRootLoc = IKFootRootTransform.GetLocation();
+	const FVector IKFootFloorNormal = IKFootRootTransform.GetUnitAxis(EAxis::Z);
+
+	if (!ensureMsgf(GravityDir.ContainsNaN() == false, TEXT("FAnimNode_SlopeWarping, Incoming GravityDir contains NaN (%s)"), *GravityDir.ToCompactString()))
+	{
+		GravityDir = -FVector::UpVector;
+	}
+	GravityDirCompSpace = MyAnimInstanceProxy->GetComponentTransform().InverseTransformVectorNoScale(GravityDir).GetSafeNormal();
+
+	// Get Feet IK Transforms
+	{
+		for (auto& FootData : FeetData)
+		{
+			FootData.IKFootTransform = Output.Pose.GetComponentSpaceTransform(FootData.IKFootBoneIndex);
+			if (!ensureMsgf(FootData.IKFootTransform.ContainsNaN() == false, TEXT("FAnimNode_SlopeWarping, Incoming IKFootTransform contains NaN (%s)"), *FootData.IKFootTransform.ToString()))
+			{
+				FootData.IKFootTransform = FTransform::Identity;
+			}
+
+			const FVector FootIKLoc = FootData.IKFootTransform.GetLocation();
+
+#if ENABLE_ANIM_DEBUG
+			if (bShowDebug)
+			{
+				const FVector FootIKWorldLoc = MyAnimInstanceProxy->GetComponentTransform().TransformPosition(FootIKLoc);
+				MyAnimInstanceProxy->AnimDrawDebugSphere(FootIKWorldLoc, 5.f, 16, FColor::Red);
+			}
+#endif
+		}
+	}
+
+	// Initial IK Feet Centroid
+	const FVector InitialCentroidLoc = GetFootIKCentroid(FeetData);
+
+#if ENABLE_ANIM_DEBUG
+	if (bShowDebug)
+	{
+		FVector CentroidWorldLoc = MyAnimInstanceProxy->GetComponentTransform().TransformPosition(InitialCentroidLoc);
+		MyAnimInstanceProxy->AnimDrawDebugSphere(CentroidWorldLoc, 5.f, 4, FColor::Red);
+	}
+#endif
+
+	FVector FloorNormal, FloorOffset;
+	GetSmoothedFloorInfo(Output.Pose, FloorNormal, FloorOffset);
+
+	if (!ensureMsgf(FloorNormal.ContainsNaN() == false, TEXT("FAnimNode_SlopeWarping, Incoming FloorNormal contains NaN (%s)"), *FloorNormal.ToCompactString()))
+	{
+		FloorNormal = FVector::UpVector;
+	}
+
+	if (!ensureMsgf(FloorOffset.ContainsNaN() == false, TEXT("FAnimNode_SlopeWarping, Incoming FloorOffset contains NaN (%s)"), *FloorOffset.ToCompactString()))
+	{
+		FloorOffset = FVector::ZeroVector;
+	}
+
+	const bool bHasFloorOffset = !FloorOffset.IsNearlyZero(KINDA_SMALL_NUMBER);
+
+#if ENABLE_ANIM_DEBUG
+	if (bShowDebug)
+	{
+		const FVector FloorOriginWorldSpace = MyAnimInstanceProxy->GetComponentTransform().TransformPosition(FloorOffset);
+		const FVector FloorNormalWorldSpace = MyAnimInstanceProxy->GetComponentTransform().TransformVectorNoScale(FloorNormal);
+		MyAnimInstanceProxy->AnimDrawDebugSphere(FloorOriginWorldSpace, 3.f, 8, FColor::Yellow);
+		MyAnimInstanceProxy->AnimDrawDebugDirectionalArrow(FloorOriginWorldSpace, FloorOriginWorldSpace + FloorNormalWorldSpace * 100.f, 10.f, FColor::Yellow);
+
+		FString DebugString = FString::Printf(TEXT("FloorOffset: %s"), *FloorOffset.ToCompactString());
+		MyAnimInstanceProxy->AnimDrawDebugOnScreenMessage(DebugString, FColor::Red);
+	}
+#endif
+
+	// Orient along floor normal
+
+	// Find Delta Rotation take takes us from Animated Slope to In-Game Slope
+	FQuat DeltaSlopeRotation = FQuat::FindBetweenNormals(IKFootFloorNormal, FloorNormal);
+	if (!ensureMsgf(DeltaSlopeRotation.ContainsNaN() == false, TEXT("FAnimNode_SlopeWarping, Incoming DeltaSlopeRotation contains NaN (%s)"), *DeltaSlopeRotation.ToString()))
+	{
+		DeltaSlopeRotation = FQuat::Identity;
+	}
+
+	// Rotate our Joint quaternion by this delta rotation
+	const bool bPoseNeedsToBeRotated = !DeltaSlopeRotation.IsIdentity();
+	if (bPoseNeedsToBeRotated)
+	{
+		IKFootRootTransform.SetRotation(DeltaSlopeRotation * IKFootRootTransform.GetRotation());
+		IKFootRootTransform.NormalizeRotation();
+	}
+	IKFootRootTransform.AddToTranslation(FloorOffset);
+
+	// Warp Feet IK to floor
+	if (bPoseNeedsToBeRotated)
+	{
+		for (auto& FootData : FeetData)
+		{
+			// Warp Location
+			const FVector WarpedFootIKLoc = DeltaSlopeRotation.RotateVector(FootData.IKFootTransform.GetLocation());
+			FootData.IKFootTransform.SetTranslation(WarpedFootIKLoc + FloorOffset);
+
+			// Warp Rotation
+			FootData.IKFootTransform.SetRotation(DeltaSlopeRotation * FootData.IKFootTransform.GetRotation());
+			FootData.IKFootTransform.NormalizeRotation();
+
+#if ENABLE_ANIM_DEBUG
+			if (bShowDebug)
+			{
+				const FVector FootIKWorldLoc = MyAnimInstanceProxy->GetComponentTransform().TransformPosition(FootData.IKFootTransform.GetLocation());
+				MyAnimInstanceProxy->AnimDrawDebugSphere(FootIKWorldLoc, 5.f, 16, FColor::Green);
+			}
+#endif
+		}
+	}
+	else if (bHasFloorOffset)
+	{
+		for (auto& FootData : FeetData)
+		{
+			FootData.IKFootTransform.AddToTranslation(FloorOffset);
+		}
+	}
+
+	// Pelvis Adjustment
+	const FCompactPoseBoneIndex PelvisBoneIndex = PelvisBone.GetCompactPoseIndex(RequiredBones);
+	FTransform PelvisTransform = Output.Pose.GetComponentSpaceTransform(PelvisBoneIndex);
+	if (!ensureMsgf(PelvisTransform.ContainsNaN() == false, TEXT("FAnimNode_SlopeWarping, Incoming PelvisTransform contains NaN (%s)"), *PelvisTransform.ToString()))
+	{
+		PelvisTransform = FTransform::Identity;
+	}
+
+	const FVector InitialPelvisLocation = PelvisTransform.GetLocation();
+
+#if ENABLE_ANIM_DEBUG
+	if (bShowDebug)
+	{
+		const FVector PelvisWorldLoc = MyAnimInstanceProxy->GetComponentTransform().TransformPosition(PelvisTransform.GetLocation());
+		MyAnimInstanceProxy->AnimDrawDebugSphere(PelvisWorldLoc, 5.f, 16, FColor::Red);
+	}
+#endif
+
+	PelvisTransform.AddToTranslation(FloorOffset);
+	const FVector PelvisLocationBeforeInterp = PelvisTransform.GetLocation();
+
+	// keep pelvis inline with feet centroid.
+	if (bPoseNeedsToBeRotated)
+	{
+		const FVector PostRotationFeetCentroid = GetFootIKCentroid(FeetData);
+		const FVector CentroidOffset = (PostRotationFeetCentroid - InitialCentroidLoc) - FloorOffset;
+
+		PelvisTransform.AddToTranslation(CentroidOffset);
+
+#if ENABLE_ANIM_DEBUG
+		if (bShowDebug)
+		{
+			MyAnimInstanceProxy->AnimDrawDebugSphere(MyAnimInstanceProxy->GetComponentTransform().TransformPosition(PostRotationFeetCentroid), 5.f, 4, FColor::Orange);
+			MyAnimInstanceProxy->AnimDrawDebugSphere(MyAnimInstanceProxy->GetComponentTransform().TransformPosition(PelvisTransform.GetLocation()), 5.f, 16, FColor::Orange);
+		}
+#endif
+	}
+
+	// pull hips down to prevent leg over-extension
+	if (bPoseNeedsToBeRotated && bPullPelvisDown)
+	{
+		// Test current pelvis offset and modify if needed
+		FVector PelvisOffset = PelvisTransform.GetLocation() - InitialPelvisLocation;
+
+		bool bCheckForLegOverextension = true;
+		int32 NumIterations = 0;
+		while (bCheckForLegOverextension && (NumIterations++ < 4))
+		{
+			bCheckForLegOverextension = false;
+			for (auto& FootData : FeetData)
+			{
+				const FVector HipLocation = Output.Pose.GetComponentSpaceTransform(FootData.HipBoneIndex).GetLocation();
+				const FVector FKFootLocation = Output.Pose.GetComponentSpaceTransform(FootData.FKFootBoneIndex).GetLocation();
+
+				const float OriginalLegLength = (HipLocation - FKFootLocation).Size();
+				const FVector IKFootToAdjustedHip = HipLocation + PelvisOffset - FootData.IKFootTransform.GetLocation();
+				const float DistanceToIKTarget = IKFootToAdjustedHip.Size();
+
+				// Pull down hips if we have to
+				if (DistanceToIKTarget > OriginalLegLength)
+				{
+					FVector NewHipLocation = FootData.IKFootTransform.GetLocation() + IKFootToAdjustedHip.GetSafeNormal() * OriginalLegLength;
+
+					PelvisOffset = NewHipLocation - HipLocation;
+
+// 					PelvisLocAfterOffset = FootData.IKFootTransform.GetLocation() + (PelvisLocAfterOffset - FootData.IKFootTransform.GetLocation()).GetSafeNormal() * OriginalLegLength;
+
+					// modified hips, take another iteration to make sure this hasn't made other feet unreachable.
+					bCheckForLegOverextension = true;
+				}
+			}
+		}
+
+		// Set our position to the modified position if it has changed
+		PelvisTransform.SetLocation(InitialPelvisLocation + PelvisOffset);
+
+#if ENABLE_ANIM_DEBUG
+		if (bShowDebug)
+		{
+			const FVector PelvisWorldLoc = MyAnimInstanceProxy->GetComponentTransform().TransformPosition(PelvisTransform.GetLocation());
+			MyAnimInstanceProxy->AnimDrawDebugSphere(PelvisWorldLoc, 5.f, 16, FColor::Emerald);
+		}
+#endif
+	}
+
+	// Adjustment to keep pose inside of capsule
+	if (bKeepMeshInsideOfCapsule)
+	{
+		// Project Pelvis transform onto slope.
+		FVector PelvisOffset = PelvisTransform.GetLocation() - PelvisLocationBeforeInterp;
+		if (!PelvisOffset.IsNearlyZero())
+		{
+			FVector ProjectedPelvisOffset = FVector::VectorPlaneProject(PelvisOffset, FloorNormal);
+
+			// Shift everything along slope to keep mesh inside of capsule.
+			PelvisTransform.AddToTranslation(-ProjectedPelvisOffset);
+
+			for (FSlopeWarpingFootData& FootData : FeetData)
+			{
+				FootData.IKFootTransform.AddToTranslation(-ProjectedPelvisOffset);
+			}
+
+			IKFootRootTransform.AddToTranslation(-ProjectedPelvisOffset);
+		}
+	}
+
+	// Interpolate Pelvis offset smoothly
+	if (true)
+	{
+		FVector PelvisOffset = PelvisTransform.GetLocation() - PelvisLocationBeforeInterp;
+		PelvisOffsetInterpolator.Update(PelvisOffset, CachedDeltaTime);
+		PelvisTransform.SetLocation(PelvisLocationBeforeInterp + PelvisOffsetInterpolator.GetPosition());
+
+#if ENABLE_ANIM_DEBUG
+		if (bShowDebug)
+		{
+			const FVector PelvisWorldLoc = MyAnimInstanceProxy->GetComponentTransform().TransformPosition(PelvisTransform.GetLocation());
+			MyAnimInstanceProxy->AnimDrawDebugSphere(PelvisWorldLoc, 5.f, 16, FColor::Purple);
+		}
+#endif
+	}
+
+	const bool bPelvisInterpolated = !PelvisOffsetInterpolator.GetPosition().IsNearlyZero(KINDA_SMALL_NUMBER);
+
+	// rotate hips towards IK feet target (to help IK nodes), and pull legs in towards pelvis to prevent over extension
+	const bool bNeedsToAdjustHips = (bPelvisInterpolated || bPoseNeedsToBeRotated);
+	if (bNeedsToAdjustHips)
+	{
+		const FVector PelvisOffset = (PelvisTransform.GetLocation() - InitialPelvisLocation);
+		for (auto& FootData : FeetData)
+		{
+			const FTransform HipTransform = Output.Pose.GetComponentSpaceTransform(FootData.HipBoneIndex);
+			const FTransform FKFootTransform = Output.Pose.GetComponentSpaceTransform(FootData.FKFootBoneIndex);
+			FTransform AdjustedHipTransform = HipTransform;
+
+			AdjustedHipTransform.AddToTranslation(PelvisOffset);
+
+			const FVector InitialDir = (FKFootTransform.GetLocation() - HipTransform.GetLocation()).GetSafeNormal();
+			const FVector TargetDir = (FootData.IKFootTransform.GetLocation() - AdjustedHipTransform.GetLocation()).GetSafeNormal();
+
+#if ENABLE_ANIM_DEBUG
+			if (bShowDebug)
+			{
+				const FTransform SkelTM = MyAnimInstanceProxy->GetComponentTransform();
+				MyAnimInstanceProxy->AnimDrawDebugLine(SkelTM.TransformPosition(HipTransform.GetLocation()), SkelTM.TransformPosition(FKFootTransform.GetLocation()), FColor::Red);
+				MyAnimInstanceProxy->AnimDrawDebugLine(SkelTM.TransformPosition(AdjustedHipTransform.GetLocation()), SkelTM.TransformPosition(FootData.IKFootTransform.GetLocation()), FColor::Green);
+			}
+#endif
+
+			if (!InitialDir.IsNearlyZero() && !TargetDir.IsNearlyZero())
+			{
+				// Find Delta Rotation take takes us from Old to New dir
+				const FQuat DeltaRotation = FQuat::FindBetweenNormals(InitialDir, TargetDir);
+
+				// Rotate our Joint quaternion by this delta rotation
+				AdjustedHipTransform.SetRotation(DeltaRotation * AdjustedHipTransform.GetRotation());
+				AdjustedHipTransform.NormalizeRotation();
+
+				// Add adjusted hip transform
+				if (ensureMsgf(AdjustedHipTransform.ContainsNaN() == false, TEXT("FAnimNode_SlopeWarping, Outgoing AdjustedHipTransform contains NaN (%s)"), *AdjustedHipTransform.ToString()))
+				{
+					OutBoneTransforms.Add(FBoneTransform(FootData.HipBoneIndex, AdjustedHipTransform));
+				}
+
+				// Clamp IK Feet bone based on FK leg. To prevent over-extension and preserve animated pose.
+				{
+					const float FKLength = FVector::Dist(FKFootTransform.GetLocation(), HipTransform.GetLocation());
+					const float IKLength = FVector::Dist(FootData.IKFootTransform.GetLocation(), AdjustedHipTransform.GetLocation());
+					if (IKLength > FKLength)
+					{
+						FootData.IKFootTransform.SetLocation(AdjustedHipTransform.GetLocation() + TargetDir * FKLength);
+					}
+				}
+			}
+		}
+	}
+
+	// Add transforms that have been modified
+	{
+		// Add adjusted IK Foot Root
+		if (bPoseNeedsToBeRotated)
+		{
+			if (ensureMsgf(IKFootRootTransform.ContainsNaN() == false, TEXT("FAnimNode_SlopeWarping, Outgoing IKFootRootTransform contains NaN (%s)"), *IKFootRootTransform.ToString()))
+			{
+				OutBoneTransforms.Add(FBoneTransform(IKFootRootBoneIndex, IKFootRootTransform));
+			}
+		}
+
+		if (bHasFloorOffset)
+		{
+			const FCompactPoseBoneIndex RootBoneIndex(0);
+			FTransform RootBoneTransform(Output.Pose.GetComponentSpaceTransform(RootBoneIndex));
+			RootBoneTransform.AddToTranslation(FloorOffset);
+
+			if (ensureMsgf(RootBoneTransform.ContainsNaN() == false, TEXT("FAnimNode_SlopeWarping, Outgoing RootBoneTransform contains NaN (%s)"), *RootBoneTransform.ToString()))
+			{
+				OutBoneTransforms.Add(FBoneTransform(RootBoneIndex, RootBoneTransform));
+			}
+		}
+
+		if (bPoseNeedsToBeRotated || bPelvisInterpolated)
+		{
+			if (ensureMsgf(PelvisTransform.ContainsNaN() == false, TEXT("FAnimNode_SlopeWarping, Outgoing PelvisTransform contains NaN (%s)"), *PelvisTransform.ToString()))
+			{
+				OutBoneTransforms.Add(FBoneTransform(PelvisBoneIndex, PelvisTransform));
+			}
+		}
+
+#if ENABLE_ANIM_DEBUG
+		if (bShowDebug)
+		{
+			FVector PelvisWorldLoc = MyAnimInstanceProxy->GetComponentTransform().TransformPosition(PelvisTransform.GetLocation());
+			MyAnimInstanceProxy->AnimDrawDebugSphere(PelvisWorldLoc, 5.f, 16, FColor::Blue);
+		}
+#endif
+
+		if (bPoseNeedsToBeRotated)
+		{
+			for (auto& FootData : FeetData)
+			{
+				if (ensureMsgf(FootData.IKFootTransform.ContainsNaN() == false, TEXT("FAnimNode_SlopeWarping, Outgoing FootData.IKFootTransform contains NaN (%s)"), *FootData.IKFootTransform.ToString()))
+				{
+					OutBoneTransforms.Add(FBoneTransform(FootData.IKFootBoneIndex, FootData.IKFootTransform));
+				}
+
+#if ENABLE_ANIM_DEBUG
+				if (bShowDebug)
+				{
+					FVector IKFootWorldLocation = MyAnimInstanceProxy->GetComponentTransform().TransformPosition(FootData.IKFootTransform.GetLocation());
+					MyAnimInstanceProxy->AnimDrawDebugSphere(IKFootWorldLocation, 5.f, 16, FColor::Blue);
+				}
+#endif
+			}
+		}
+
+#if ENABLE_ANIM_DEBUG
+		if (bShowDebug)
+		{
+			FVector FinalCentroidLoc = GetFootIKCentroid(FeetData);
+			FVector CentroidWorldLoc = MyAnimInstanceProxy->GetComponentTransform().TransformPosition(FinalCentroidLoc);
+			MyAnimInstanceProxy->AnimDrawDebugSphere(CentroidWorldLoc, 5.f, 4, FColor::Blue);
+		}
+#endif
+
+	}
+
+	// Sort OutBoneTransforms so indices are in increasing order.
+	OutBoneTransforms.Sort(FCompareBoneTransformIndex());
+
+	// Clear time accumulator, to be filled during next update.
+	CachedDeltaTime = 0.f;
+}
+
+bool FAnimNode_SlopeWarping::IsValidToEvaluate(const USkeleton* Skeleton, const FBoneContainer& RequiredBones)
+{
+	const bool bIsEnabled = (CVarSlopeWarpingEnable.GetValueOnAnyThread() == 1);
+	return MyMovementComponent && bIsEnabled && (IKFootRootBone.GetCompactPoseIndex(RequiredBones) != INDEX_NONE) && (PelvisBone.GetCompactPoseIndex(RequiredBones) != INDEX_NONE) && (FeetData.Num() > 0);
+}
+
+static FCompactPoseBoneIndex FindHipBoneIndex(const FCompactPoseBoneIndex& InFootBoneIndex, const int32& NumBonesInLimb, const FBoneContainer& RequiredBones)
+{
+	FCompactPoseBoneIndex BoneIndex = InFootBoneIndex;
+	if (BoneIndex != INDEX_NONE)
+	{
+		FCompactPoseBoneIndex ParentBoneIndex = RequiredBones.GetParentBoneIndex(BoneIndex);
+
+		int32 NumIterations = NumBonesInLimb;
+		while ((NumIterations-- > 0) && (ParentBoneIndex != INDEX_NONE))
+		{
+			BoneIndex = ParentBoneIndex;
+			ParentBoneIndex = RequiredBones.GetParentBoneIndex(BoneIndex);
+		};
+	}
+
+	return BoneIndex;
+}
+
+/** OutBoneIndex and OutLocalTransform will only be set if found! */
+static void SetBoneIndexAndLocalTransformFromSocket(const FName& InSocketName, const USkeletalMeshComponent& SkelMeshComp, const FBoneContainer& RequiredBones, FCompactPoseBoneIndex& OutBoneIndex, FTransform& OutLocalTransform)
+{
+	if (InSocketName != NAME_None)
+	{
+		const USkeletalMeshSocket* SkelMeshSocket = SkelMeshComp.GetSocketByName(InSocketName);
+		if (SkelMeshSocket)
+		{
+			OutBoneIndex = RequiredBones.MakeCompactPoseIndex(FMeshPoseBoneIndex(RequiredBones.GetPoseBoneIndexForBoneName(SkelMeshSocket->BoneName)));
+			OutLocalTransform = SkelMeshSocket->GetSocketLocalTransform();
+		}
+	}
+}
+
+void FAnimNode_SlopeWarping::InitializeBoneReferences(const FBoneContainer& RequiredBones)
+{
+	IKFootRootBone.Initialize(RequiredBones);
+
+	PelvisBone.Initialize(RequiredBones);
+
+	FeetData.Reset();
+	check(MySkelMeshComponent);
+	for (auto& FootDef : FeetDefinitions)
+	{
+		FootDef.IKFootBone.Initialize(RequiredBones);
+		FootDef.FKFootBone.Initialize(RequiredBones);
+
+		FSlopeWarpingFootData FootData;
+		FootData.IKFootBoneIndex = FootDef.IKFootBone.GetCompactPoseIndex(RequiredBones);
+		FootData.FKFootBoneIndex = FootDef.FKFootBone.GetCompactPoseIndex(RequiredBones);
+		FootData.HipBoneIndex = FindHipBoneIndex(FootData.FKFootBoneIndex, FMath::Max(FootDef.NumBonesInLimb, 1), RequiredBones);
+
+		if ((FootData.IKFootBoneIndex != INDEX_NONE)
+			&& (FootData.FKFootBoneIndex != INDEX_NONE)
+			&& (FootData.HipBoneIndex != INDEX_NONE))
+		{
+			FootData.FootDefinitionPtr = &FootDef;
+			FeetData.Add(FootData);
+		}
+	}
+}
diff --git a/Engine/Plugins/Animation/AnimationWarping/Source/Runtime/Private/BoneControllers/AnimNode_StrideWarping.cpp b/Engine/Plugins/Animation/AnimationWarping/Source/Runtime/Private/BoneControllers/AnimNode_StrideWarping.cpp
new file mode 100644
--- /dev/null
+++ b/Engine/Plugins/Animation/AnimationWarping/Source/Runtime/Private/BoneControllers/AnimNode_StrideWarping.cpp
@@ -0,0 +1,406 @@
+// Copyright Epic Games, Inc. All Rights Reserved.
+
+#include "BoneControllers/AnimNode_StrideWarping.h"
+#include "Animation/AnimInstanceProxy.h"
+#include "Animation/AnimRootMotionProvider.h"
+
+DECLARE_CYCLE_STAT(TEXT("StrideWarping Eval"), STAT_StrideWarping_Eval, STATGROUP_Anim);
+
+#if ENABLE_ANIM_DEBUG
+static TAutoConsoleVariable<int32> CVarAnimNodeStrideWarpingDebug(TEXT("a.AnimNode.StrideWarping.Debug"), 0, TEXT("Turn on visualization debugging for Stride Warping"));
+static TAutoConsoleVariable<int32> CVarAnimNodeStrideWarpingVerbose(TEXT("a.AnimNode.StrideWarping.Verbose"), 0, TEXT("Turn on verbose graph debugging for Stride Warping"));
+static TAutoConsoleVariable<int32> CVarAnimNodeStrideWarpingEnable(TEXT("a.AnimNode.StrideWarping.Enable"), 1, TEXT("Toggle Stride Warping"));
+#endif
+
+void FAnimNode_StrideWarping::GatherDebugData(FNodeDebugData& DebugData)
+{
+	FString DebugLine = DebugData.GetNodeName(this);
+#if ENABLE_ANIM_DEBUG
+	if (CVarAnimNodeStrideWarpingVerbose.GetValueOnAnyThread() == 1)
+	{ 
+		DebugLine += FString::Printf(TEXT("\n - Evaluation Mode: (%s)"), Mode == EWarpingEvaluationMode::Graph ? TEXT("Graph") : TEXT("Manual"));
+		DebugLine += FString::Printf(TEXT("\n - Stride Scale: (%.3fd)"), ActualStrideScale);
+		DebugLine += FString::Printf(TEXT("\n - Stride Direction: (%s)"), *(ActualStrideDirection.ToCompactString()));
+		if (Mode == EWarpingEvaluationMode::Graph)
+		{
+			DebugLine += FString::Printf(TEXT("\n - Locomotion Speed: (%.3fd)"), LocomotionSpeed);
+#if WITH_EDITORONLY_DATA
+			DebugLine += FString::Printf(TEXT("\n - Root Motion Direction: (%s)"), *(CachedRootMotionDeltaTranslation.GetSafeNormal().ToCompactString()));
+			DebugLine += FString::Printf(TEXT("\n - Root Motion Speed: (%.3fd)"), CachedRootMotionDeltaSpeed);
+#endif
+			DebugLine += FString::Printf(TEXT("\n - Min Locomotion Speed Threshold: (%.3fd)"), MinLocomotionSpeedThreshold);
+		}
+		DebugLine += FString::Printf(TEXT("\n - Floor Normal: (%s)"), *(FloorNormalDirection.Value.ToCompactString()));
+		DebugLine += FString::Printf(TEXT("\n - Gravity Direction: (%s)"), *(GravityDirection.Value.ToCompactString()));
+	}
+	else
+#endif
+	{
+		DebugLine += FString::Printf(TEXT("(Stride Scale: %.3fd, Direction: %s)"), ActualStrideScale, *(ActualStrideDirection.ToCompactString()));
+	}
+	DebugData.AddDebugItem(DebugLine);
+	ComponentPose.GatherDebugData(DebugData);
+}
+
+void FAnimNode_StrideWarping::Initialize_AnyThread(const FAnimationInitializeContext& Context)
+{
+	FAnimNode_SkeletalControlBase::Initialize_AnyThread(Context);
+	
+	AnimInstanceProxy = Context.AnimInstanceProxy;
+	StrideScaleModifierState.Reinitialize();
+	ActualStrideDirection = FVector::ForwardVector;
+	ActualStrideScale = 1.f;
+#if WITH_EDITORONLY_DATA
+	CachedRootMotionDeltaTranslation = FVector::ZeroVector;
+	CachedRootMotionDeltaSpeed = 0.f;
+#endif
+}
+
+void FAnimNode_StrideWarping::UpdateInternal(const FAnimationUpdateContext& Context)
+{
+	FAnimNode_SkeletalControlBase::UpdateInternal(Context);
+	CachedDeltaTime = Context.GetDeltaTime();
+}
+
+void FAnimNode_StrideWarping::EvaluateSkeletalControl_AnyThread(FComponentSpacePoseContext& Output, TArray<FBoneTransform>& OutBoneTransforms)
+{
+	SCOPE_CYCLE_COUNTER(STAT_StrideWarping_Eval);
+	check(OutBoneTransforms.IsEmpty());
+
+	ActualStrideDirection = StrideDirection;
+	ActualStrideScale = StrideScale;
+
+	bool bGraphDrivenWarping = false;
+	const UE::Anim::IAnimRootMotionProvider* RootMotionProvider = UE::Anim::IAnimRootMotionProvider::Get();
+
+	if (Mode == EWarpingEvaluationMode::Graph)
+	{
+		bGraphDrivenWarping = !!RootMotionProvider;
+		ensureMsgf(bGraphDrivenWarping, TEXT("Graph driven Stride Warping expected a valid root motion delta provider interface."));
+	}
+
+	FTransform RootMotionTransformDelta = FTransform::Identity;
+#if !WITH_EDITORONLY_DATA
+	FVector CachedRootMotionDeltaTranslation = FVector::ZeroVector;
+#endif
+	if (bGraphDrivenWarping)
+	{
+		// Graph driven stride warping will override the manual stride direction with the intent of the current animation sub-graph's accumulated root motion
+		bGraphDrivenWarping = RootMotionProvider->ExtractRootMotion(Output.CustomAttributes, RootMotionTransformDelta);
+		if (ensureMsgf(bGraphDrivenWarping, TEXT(
+			"Graph driven Stride Warping expected a root motion delta to be present in the attribute stream.\n"
+			"Check that the evaluating Animation Sequence(s) have: \"Enable Root Motion\" enabled.")))
+		{
+			CachedRootMotionDeltaTranslation = RootMotionTransformDelta.GetTranslation();
+			ActualStrideDirection = CachedRootMotionDeltaTranslation.GetSafeNormal();
+		}
+		else
+		{
+			// Early exit on broken graph driven behavior 
+			return;
+		}
+	}
+
+	const FBoneContainer& RequiredBones = Output.Pose.GetPose().GetBoneContainer();
+	const FTransform IKFootRootTransform = Output.Pose.GetComponentSpaceTransform(IKFootRootBone.GetCompactPoseIndex(RequiredBones));
+	const FVector ResolvedFloorNormal = FloorNormalDirection.AsComponentSpaceDirection(AnimInstanceProxy, IKFootRootTransform);
+	const FVector ResolvedGravityDirection = GravityDirection.AsComponentSpaceDirection(AnimInstanceProxy, IKFootRootTransform);
+
+	if (bOrientStrideDirectionUsingFloorNormal)
+	{
+		const FVector StrideWarpingAxis = ResolvedFloorNormal ^ ActualStrideDirection;
+		ActualStrideDirection = StrideWarpingAxis ^ ResolvedFloorNormal;
+	}
+
+#if ENABLE_ANIM_DEBUG
+	bool bDebugging = false;
+#if WITH_EDITORONLY_DATA
+	bDebugging = bDebugging || bEnableDebugDraw;
+#else
+	constexpr float DebugDrawScale = 1.f;
+#endif
+	bDebugging = bDebugging || CVarAnimNodeStrideWarpingDebug.GetValueOnAnyThread() == 1;
+	if (bDebugging)
+	{
+		// Draw Floor Normal
+		AnimInstanceProxy->AnimDrawDebugDirectionalArrow(
+			AnimInstanceProxy->GetComponentTransform().TransformPosition(IKFootRootTransform.GetLocation()),
+			AnimInstanceProxy->GetComponentTransform().TransformPosition(IKFootRootTransform.GetLocation() + ResolvedFloorNormal * 25.f * DebugDrawScale),
+			40.f * DebugDrawScale, FColor::Blue, false, 0.f, 2.f * DebugDrawScale);
+	}
+#endif
+
+	// Get all foot IK Transforms
+	for (auto& Foot : FootData)
+	{
+		Foot.IKFootBoneTransform = Output.Pose.GetComponentSpaceTransform(Foot.IKFootBoneIndex);
+#if ENABLE_ANIM_DEBUG
+		if (bDebugging
+#if WITH_EDITORONLY_DATA
+			&& bDebugDrawIKFootOrigin
+#endif
+		)
+		{
+			const FVector FootWorldLocation = AnimInstanceProxy->GetComponentTransform().TransformPosition(Foot.IKFootBoneTransform.GetLocation());
+			AnimInstanceProxy->AnimDrawDebugSphere(FootWorldLocation, 8.f * DebugDrawScale, 16, FColor::Red);
+		}
+#endif
+	}
+
+	if (bGraphDrivenWarping)
+	{
+#if WITH_EDITORONLY_DATA
+		CachedRootMotionDeltaSpeed = CachedRootMotionDeltaTranslation.Size() / CachedDeltaTime;
+#else
+		const float CachedRootMotionDeltaSpeed = CachedRootMotionDeltaTranslation.Size() / CachedDeltaTime;
+#endif
+		// Graph driven stride scale factor will be determined by the ratio of the
+		// locomotion (capsule/physics) speed against the animation root motion speed
+		ActualStrideScale = LocomotionSpeed / CachedRootMotionDeltaSpeed;
+	}
+
+	// Allow the opportunity for stride scale clamping and interpolation regardless of evaluation mode
+	ActualStrideScale = StrideScaleModifierState.ApplyTo(StrideScaleModifier, ActualStrideScale, CachedDeltaTime);
+
+	if (bGraphDrivenWarping)
+	{
+		// Forward the side effects of stride warping on the root motion contribution for this sub-graph
+		RootMotionTransformDelta.ScaleTranslation(ActualStrideScale);
+		const bool bRootMotionOverridden = RootMotionProvider->OverrideRootMotion(RootMotionTransformDelta, Output.CustomAttributes);
+		ensureMsgf(bRootMotionOverridden, TEXT("Graph driven Stride Warping expected a root motion delta to be present in the attribute stream prior to warping/overriding it."));
+	}
+
+	// Scale IK feet bones along Stride Warping Axis, from the Thigh bone location.
+	for (auto& Foot : FootData)
+	{
+		// Stride Warping along Stride Warping Axis
+		const FVector IKFootLocation = Foot.IKFootBoneTransform.GetLocation();
+		const FVector ThighBoneLocation = Output.Pose.GetComponentSpaceTransform(Foot.ThighBoneIndex).GetLocation();
+
+		// Project Thigh Bone Location on plane made of FootIKLocation and FloorPlaneNormal, along Gravity Dir.
+		// This will be the StrideWarpingPlaneOrigin
+		const FVector StrideWarpingPlaneOrigin = (FMath::Abs(ResolvedGravityDirection | ResolvedFloorNormal) > DELTA) ? FMath::LinePlaneIntersection(ThighBoneLocation, ThighBoneLocation + ResolvedGravityDirection, IKFootLocation, ResolvedFloorNormal) : IKFootLocation;
+
+		// Project FK Foot along StrideWarping Plane, this will be our Scale Origin
+		const FVector ScaleOrigin = FVector::PointPlaneProject(IKFootLocation, StrideWarpingPlaneOrigin, ActualStrideDirection);
+
+		// Now the ScaleOrigin and IKFootLocation are forming a line parallel to the floor, and we can scale the IK foot.
+		const FVector WarpedLocation = ScaleOrigin + (IKFootLocation - ScaleOrigin) * ActualStrideScale;
+		Foot.IKFootBoneTransform.SetLocation(WarpedLocation);
+
+#if ENABLE_ANIM_DEBUG
+		if (bDebugging
+#if WITH_EDITORONLY_DATA
+			&& bDebugDrawIKFootAdjustment
+#endif
+		)
+		{
+			const FVector FootWorldLocation = AnimInstanceProxy->GetComponentTransform().TransformPosition(Foot.IKFootBoneTransform.GetLocation());
+			AnimInstanceProxy->AnimDrawDebugSphere(FootWorldLocation, 8.f * DebugDrawScale, 16, FColor::Green);
+
+			const FVector ScaleOriginWorldLoc = AnimInstanceProxy->GetComponentTransform().TransformPosition(ScaleOrigin);
+			AnimInstanceProxy->AnimDrawDebugSphere(ScaleOriginWorldLoc, 8.f * DebugDrawScale, 16, FColor::Yellow);
+		}
+#endif
+	}
+
+	FVector PelvisOffset = FVector::ZeroVector;
+	const FCompactPoseBoneIndex PelvisBoneIndex = PelvisBone.GetCompactPoseIndex(RequiredBones);
+
+	FTransform PelvisTransform = Output.Pose.GetComponentSpaceTransform(PelvisBoneIndex);
+	const FVector InitialPelvisLocation = PelvisTransform.GetLocation();
+
+	TArray<float, TInlineAllocator<10>> FKFootDistancesToPelvis;
+	FKFootDistancesToPelvis.Reserve(FootData.Num());
+
+	TArray<FVector, TInlineAllocator<10>> IKFootLocations;
+	IKFootLocations.Reserve(FootData.Num());
+	
+	for (auto& Foot : FootData)
+	{
+		const FVector FKFootLocation = Output.Pose.GetComponentSpaceTransform(Foot.FKFootBoneIndex).GetLocation();
+		FKFootDistancesToPelvis.Add(FVector::Dist(FKFootLocation, InitialPelvisLocation));
+
+		const FVector IKFootLocation = Foot.IKFootBoneTransform.GetLocation();
+		IKFootLocations.Add(IKFootLocation);
+	}
+
+	// Adjust Pelvis down if needed to keep foot contact with the ground and prevent over-extension
+	PelvisTransform = PelvisIKFootSolver.Solve(PelvisTransform, FKFootDistancesToPelvis, IKFootLocations, CachedDeltaTime);
+
+#if ENABLE_ANIM_DEBUG
+	if (bDebugging
+#if WITH_EDITORONLY_DATA
+		&& bDebugDrawPelvisAdjustment
+#endif
+	)
+	{
+		// Draw Adjustments in Pelvis location
+		AnimInstanceProxy->AnimDrawDebugSphere(AnimInstanceProxy->GetComponentTransform().TransformPosition(InitialPelvisLocation), 8.f * DebugDrawScale, 16, FColor::Red);
+		AnimInstanceProxy->AnimDrawDebugSphere(AnimInstanceProxy->GetComponentTransform().TransformPosition(PelvisTransform.GetLocation()), 8.f * DebugDrawScale, 16, FColor::Blue);
+	}
+
+	if (bDebugging)
+	{
+		// Draw Stride Direction
+		AnimInstanceProxy->AnimDrawDebugDirectionalArrow(
+			AnimInstanceProxy->GetComponentTransform().TransformPosition(InitialPelvisLocation),
+			AnimInstanceProxy->GetComponentTransform().TransformPosition(InitialPelvisLocation + ActualStrideDirection * ActualStrideScale * 100.f * DebugDrawScale),
+			40.f * DebugDrawScale, FColor::Red, false, 0.f, 2.f * DebugDrawScale);
+	}
+#endif
+	// Add adjusted pelvis transform
+	check(!PelvisTransform.ContainsNaN());
+	OutBoneTransforms.Add(FBoneTransform(PelvisBoneIndex, PelvisTransform));
+
+	// Compute final offset to use below
+	PelvisOffset = (PelvisTransform.GetLocation() - InitialPelvisLocation);
+
+	// Rotate Thigh bones to help IK, and maintain leg shape.
+	if (bCompensateIKUsingFKThighRotation)
+	{
+		for (auto& Foot : FootData)
+		{
+			const FTransform ThighTransform = Output.Pose.GetComponentSpaceTransform(Foot.ThighBoneIndex);
+			const FTransform FKFootTransform = Output.Pose.GetComponentSpaceTransform(Foot.FKFootBoneIndex);
+			
+			FTransform AdjustedThighTransform = ThighTransform;
+			AdjustedThighTransform.AddToTranslation(PelvisOffset);
+
+			const FVector InitialDir = (FKFootTransform.GetLocation() - ThighTransform.GetLocation()).GetSafeNormal();
+			const FVector TargetDir = (Foot.IKFootBoneTransform.GetLocation() - AdjustedThighTransform.GetLocation()).GetSafeNormal();
+			
+#if ENABLE_ANIM_DEBUG
+			if (bDebugging
+#if WITH_EDITORONLY_DATA
+				&& bDebugDrawThighAdjustment
+#endif
+			)
+			{
+				AnimInstanceProxy->AnimDrawDebugLine(
+					AnimInstanceProxy->GetComponentTransform().TransformPosition(ThighTransform.GetLocation()),
+					AnimInstanceProxy->GetComponentTransform().TransformPosition(FKFootTransform.GetLocation()),
+					FColor::Red, false, 0.f, 2.f * DebugDrawScale);
+
+				AnimInstanceProxy->AnimDrawDebugLine(
+					AnimInstanceProxy->GetComponentTransform().TransformPosition(AdjustedThighTransform.GetLocation()),
+					AnimInstanceProxy->GetComponentTransform().TransformPosition(Foot.IKFootBoneTransform.GetLocation()),
+					FColor::Green, false, 0.f, 2.f * DebugDrawScale);
+			}
+#endif
+			// Find Delta Rotation take takes us from Old to New dir
+			const FQuat DeltaRotation = FQuat::FindBetweenNormals(InitialDir, TargetDir);
+
+			// Rotate our Joint quaternion by this delta rotation
+			AdjustedThighTransform.SetRotation(DeltaRotation * AdjustedThighTransform.GetRotation());
+
+			// Add adjusted thigh transform
+			check(!AdjustedThighTransform.ContainsNaN());
+			OutBoneTransforms.Add(FBoneTransform(Foot.ThighBoneIndex, AdjustedThighTransform));
+
+			// Clamp IK Feet bone based on FK leg. To prevent over-extension and preserve animated motion.
+			if (bClampIKUsingFKLimits)
+			{
+				const float FKLength = FVector::Dist(FKFootTransform.GetLocation(), ThighTransform.GetLocation());
+				const float IKLength = FVector::Dist(Foot.IKFootBoneTransform.GetLocation(), AdjustedThighTransform.GetLocation());
+				if (IKLength > FKLength)
+				{
+					const FVector ClampedFootLocation = AdjustedThighTransform.GetLocation() + TargetDir * FKLength;
+					Foot.IKFootBoneTransform.SetLocation(ClampedFootLocation);
+				}
+			}
+		}
+	}
+
+	// Add final IK feet transforms
+	for (auto& Foot : FootData)
+	{
+#if ENABLE_ANIM_DEBUG
+		if (bDebugging
+#if WITH_EDITORONLY_DATA
+			&& bDebugDrawIKFootFinal
+#endif
+		)
+		{
+			AnimInstanceProxy->AnimDrawDebugSphere(AnimInstanceProxy->GetComponentTransform().TransformPosition(Foot.IKFootBoneTransform.GetLocation()), 8.f * DebugDrawScale, 16, FColor::Blue);
+		}
+#endif
+		check(!Foot.IKFootBoneTransform.ContainsNaN());
+		OutBoneTransforms.Add(FBoneTransform(Foot.IKFootBoneIndex, Foot.IKFootBoneTransform));
+	}
+
+	// Sort OutBoneTransforms so indices are in increasing order.
+	OutBoneTransforms.Sort(FCompareBoneTransformIndex());
+}
+
+bool FAnimNode_StrideWarping::IsValidToEvaluate(const USkeleton* Skeleton, const FBoneContainer& RequiredBones)
+{
+#if ENABLE_ANIM_DEBUG
+	if (CVarAnimNodeStrideWarpingEnable.GetValueOnAnyThread() == 0)
+	{
+		return false;
+	}
+#endif
+	if (!PelvisBone.IsValidToEvaluate(RequiredBones))
+	{
+		return false;
+	}
+
+	if (!IKFootRootBone.IsValidToEvaluate(RequiredBones))
+	{
+		return false;
+	}
+
+	if (FootData.IsEmpty())
+	{
+		return false;
+	}
+	else
+	{
+		for (const auto& Foot : FootData)
+		{
+			if (Foot.IKFootBoneIndex == INDEX_NONE || Foot.FKFootBoneIndex == INDEX_NONE || Foot.ThighBoneIndex == INDEX_NONE)
+			{
+				return false;
+			}
+		}
+	}
+
+	if (Mode == EWarpingEvaluationMode::Manual)
+	{
+		if (FMath::IsNearlyEqual(StrideScaleModifierState.ApplyTo(StrideScaleModifier, StrideScale, CachedDeltaTime), 1.f, KINDA_SMALL_NUMBER))
+		{
+			return false;
+		}
+	}
+	else if (LocomotionSpeed <= MinLocomotionSpeedThreshold)
+	{
+		return false;
+	}
+
+	return true;
+}
+
+void FAnimNode_StrideWarping::InitializeBoneReferences(const FBoneContainer& RequiredBones)
+{
+	IKFootRootBone.Initialize(RequiredBones);
+	PelvisBone.Initialize(RequiredBones);
+	FootData.Empty();
+
+	for (auto& Foot : FootDefinitions)
+	{
+		Foot.IKFootBone.Initialize(RequiredBones);
+		Foot.FKFootBone.Initialize(RequiredBones);
+		Foot.ThighBone.Initialize(RequiredBones);
+
+		FStrideWarpingFootData StrideFootData;
+		StrideFootData.IKFootBoneIndex = Foot.IKFootBone.GetCompactPoseIndex(RequiredBones);
+		StrideFootData.FKFootBoneIndex = Foot.FKFootBone.GetCompactPoseIndex(RequiredBones);
+		StrideFootData.ThighBoneIndex = Foot.ThighBone.GetCompactPoseIndex(RequiredBones);
+
+		if ((StrideFootData.IKFootBoneIndex != INDEX_NONE) && (StrideFootData.FKFootBoneIndex != INDEX_NONE) && (StrideFootData.ThighBoneIndex != INDEX_NONE))
+		{
+			FootData.Add(StrideFootData);
+		}
+	}
+}
\ No newline at end of file
diff --git a/Engine/Plugins/Animation/AnimationWarping/Source/Runtime/Public/AnimationWarpingRuntimeModule.h b/Engine/Plugins/Animation/AnimationWarping/Source/Runtime/Public/AnimationWarpingRuntimeModule.h
new file mode 100644
--- /dev/null
+++ b/Engine/Plugins/Animation/AnimationWarping/Source/Runtime/Public/AnimationWarpingRuntimeModule.h
@@ -0,0 +1,12 @@
+// Copyright Epic Games, Inc. All Rights Reserved.
+
+#pragma once
+
+#include "CoreMinimal.h"
+#include "Modules/ModuleInterface.h"
+#include "Modules/ModuleManager.h"
+
+class IAnimationWarpingRuntimeModule : public IModuleInterface
+{
+};
+
diff --git a/Engine/Plugins/Animation/AnimationWarping/Source/Runtime/Public/BoneControllers/AnimNode_OrientationWarping.h b/Engine/Plugins/Animation/AnimationWarping/Source/Runtime/Public/BoneControllers/AnimNode_OrientationWarping.h
new file mode 100644
--- /dev/null
+++ b/Engine/Plugins/Animation/AnimationWarping/Source/Runtime/Public/BoneControllers/AnimNode_OrientationWarping.h
@@ -0,0 +1,144 @@
+// Copyright Epic Games, Inc. All Rights Reserved.
+
+#pragma once
+
+#include "BoneControllers/BoneControllerTypes.h"
+#include "BoneControllers/AnimNode_SkeletalControlBase.h"
+#include "AnimNode_OrientationWarping.generated.h"
+
+USTRUCT(BlueprintInternalUseOnly)
+struct ANIMATIONWARPINGRUNTIME_API FAnimNode_OrientationWarping : public FAnimNode_SkeletalControlBase
+{
+	GENERATED_BODY()
+
+	// Orientation warping evaluation mode (Graph or Manual)
+	UPROPERTY(EditAnywhere, Category=Evaluation)
+	EWarpingEvaluationMode Mode = EWarpingEvaluationMode::Manual;
+
+	// The desired orientation angle (in degrees) to warp by relative to the specified RotationAxis
+	UPROPERTY(EditAnywhere, BlueprintReadWrite, Category=Evaluation, meta=(PinShownByDefault))
+	float OrientationAngle = 0.f;
+
+	// The character locomotion angle (in degrees) relative to the specified RotationAxis
+	// This will be used in the following equation to compute the effective orientation warping angle: [Orientation = RotationBetween(RootMotionDirection, LocomotionDirection)]
+	UPROPERTY(EditAnywhere, BlueprintReadWrite, Category=Evaluation, meta=(PinShownByDefault))
+	float LocomotionAngle = 0.f;
+
+	// Specifies an angle threshold to prevent erroneous over-rotation of the character, disabled with a value of 0
+	//
+	// When the effective orientation warping angle is detected to be greater than this value (default: 90 degrees) the locomotion direction will be inverted prior to warping
+	// This will result in the following equation: [Orientation = RotationBetween(RootMotionDirection, -LocomotionDirection)]
+	//
+	// Example: Playing a forward running animation while the motion is going backward 
+	// Rather than orientation warping by 180 degrees, the system will warp by 0 degrees 
+	UPROPERTY(EditAnywhere, BlueprintReadWrite, Category=Evaluation, meta=(PinHiddenByDefault), meta=(ClampMin="0.0", ClampMax="180.0"))
+	float LocomotionAngleDeltaThreshold = 90.f;
+
+	// Spine bone definitions
+	// Used to counter the rotation of the body to keep the character facing towards the specified angle (Orientation Angle, Locomotion Angle)
+	// The amount of counter rotation applied is driven by DistributedBoneOrientationAlpha
+	UPROPERTY(EditAnywhere, Category=Settings)
+	TArray<FBoneReference> SpineBones;
+
+	// IK Foot Root Bone definition
+	UPROPERTY(EditAnywhere, Category=Settings, meta=(DisplayName="IK Foot Root Bone"))
+	FBoneReference IKFootRootBone;
+
+	// IK Foot definitions
+	UPROPERTY(EditAnywhere, Category=Settings, meta=(DisplayName="IK Foot Bones"))
+	TArray<FBoneReference> IKFootBones;
+
+	// Specifies how much rotation is applied to the character body versus IK feet
+	UPROPERTY(EditAnywhere, Category=Settings, meta=(ClampMin="0.0", ClampMax="1.0"))
+	float DistributedBoneOrientationAlpha = 0.5f;
+
+	// Rotation axis used when rotating the character body
+	UPROPERTY(EditAnywhere, Category=Settings)
+	TEnumAsByte<EAxis::Type> RotationAxis = EAxis::Z;
+
+	// Specifies the interpolation speed (in Alpha per second) towards reaching the final warped rotation angle
+	// A value of 0 will cause instantaneous warping, while a greater value will introduce smoothing
+	UPROPERTY(EditAnywhere, Category=Settings, meta=(ClampMin="0.0"))
+	float RotationInterpSpeed = 10.f;
+
+#if WITH_EDITORONLY_DATA
+	// Scale all debug drawing visualization by a factor
+	UPROPERTY(EditAnywhere, Category=Debug, meta=(ClampMin="0.0"))
+	float DebugDrawScale = 1.f;
+
+	// Enable/Disable orientation warping debug drawing
+	UPROPERTY(EditAnywhere, Category=Debug)
+	bool bEnableDebugDraw = false;
+#endif
+
+public:
+	// FAnimNode_Base interface
+	virtual void GatherDebugData(FNodeDebugData& DebugData) override;
+	virtual void UpdateInternal(const FAnimationUpdateContext& Context) override;
+	// End of FAnimNode_Base interface
+
+	// FAnimNode_SkeletalControlBase interface
+	virtual void Initialize_AnyThread(const FAnimationInitializeContext& Context) override;
+	virtual void EvaluateSkeletalControl_AnyThread(FComponentSpacePoseContext& Output, TArray<FBoneTransform>& OutBoneTransforms) override;
+	virtual bool IsValidToEvaluate(const USkeleton* Skeleton, const FBoneContainer& RequiredBones) override;
+	// End of FAnimNode_SkeletalControlBase interface
+
+private:
+	// FAnimNode_SkeletalControlBase interface
+	virtual void InitializeBoneReferences(const FBoneContainer& RequiredBones) override;
+	// End of FAnimNode_SkeletalControlBase interface
+
+	struct FOrientationWarpingSpineBoneData
+	{
+		FCompactPoseBoneIndex BoneIndex;
+		float Weight;
+
+		FOrientationWarpingSpineBoneData()
+			: BoneIndex(INDEX_NONE)
+			, Weight(0.f)
+		{
+		}
+
+		FOrientationWarpingSpineBoneData(FCompactPoseBoneIndex InBoneIndex)
+			: BoneIndex(InBoneIndex)
+			, Weight(0.f)
+		{
+		}
+
+		// Comparison Operator for Sorting
+		struct FCompareBoneIndex
+		{
+			FORCEINLINE bool operator()(const FOrientationWarpingSpineBoneData& A, const FOrientationWarpingSpineBoneData& B) const
+			{
+				return A.BoneIndex < B.BoneIndex;
+			}
+		};
+	};
+
+	struct FOrientationWarpingFootData
+	{
+		TArray<FCompactPoseBoneIndex> IKFootBoneIndexArray;
+		FCompactPoseBoneIndex IKFootRootBoneIndex;
+
+		FOrientationWarpingFootData()
+			: IKFootBoneIndexArray()
+			, IKFootRootBoneIndex(INDEX_NONE)
+		{
+		}
+	};
+
+	// Computed spine bone indices and alpha weights for the specified spine definition 
+	TArray<FOrientationWarpingSpineBoneData> SpineBoneDataArray;
+
+	// Computed IK bone indices for the specified foot definitions 
+	FOrientationWarpingFootData IKFootData;
+
+	// Internal previous frame root motion delta direction
+	FVector PreviousRootMotionDeltaDirection = FVector::ZeroVector;
+
+	// Internal previous frame orientation warping angle
+	float PreviousOrientationAngle = 0.f;
+	
+	// Internal orientation warping angle
+	float ActualOrientationAngle = 0.f;
+};
\ No newline at end of file
diff --git a/Engine/Plugins/Animation/AnimationWarping/Source/Runtime/Public/BoneControllers/AnimNode_SlopeWarping.h b/Engine/Plugins/Animation/AnimationWarping/Source/Runtime/Public/BoneControllers/AnimNode_SlopeWarping.h
new file mode 100644
--- /dev/null
+++ b/Engine/Plugins/Animation/AnimationWarping/Source/Runtime/Public/BoneControllers/AnimNode_SlopeWarping.h
@@ -0,0 +1,163 @@
+// Copyright Epic Games, Inc. All Rights Reserved.
+
+#pragma once
+
+#include "BoneControllers/BoneControllerTypes.h"
+#include "BoneControllers/AnimNode_SkeletalControlBase.h"
+#include "Engine/SpringInterpolator.h"
+#include "AnimNode_SlopeWarping.generated.h"
+
+/** Per foot definitions */
+USTRUCT()
+struct FSlopeWarpingFootDefinition
+{
+	GENERATED_USTRUCT_BODY()
+
+	UPROPERTY(EditAnywhere, Category = "Settings")
+	FBoneReference IKFootBone;
+
+	UPROPERTY(EditAnywhere, Category = "Settings")
+	FBoneReference FKFootBone;
+
+	UPROPERTY(EditAnywhere, Category = "Settings")
+	int32 NumBonesInLimb;
+
+	UPROPERTY(EditAnywhere, Category = "Settings")
+	float FootSize;
+
+	FSlopeWarpingFootDefinition()
+		: NumBonesInLimb(2)
+		, FootSize(4.f)
+	{}
+};
+
+/** Runtime foot data after validation, we guarantee these bones to exist */
+USTRUCT()
+struct FSlopeWarpingFootData
+{
+	GENERATED_USTRUCT_BODY()
+
+public:
+	FCompactPoseBoneIndex IKFootBoneIndex;
+	FCompactPoseBoneIndex FKFootBoneIndex;
+	FCompactPoseBoneIndex HipBoneIndex;
+
+	FTransform IKFootTransform;
+
+	FVector HitLocation;
+	FVector HitNormal;
+
+	FSlopeWarpingFootDefinition* FootDefinitionPtr;
+
+public:
+	FSlopeWarpingFootData()
+		: IKFootBoneIndex(INDEX_NONE)
+		, FKFootBoneIndex(INDEX_NONE)
+		, HipBoneIndex(INDEX_NONE)
+		, FootDefinitionPtr(nullptr)
+	{}
+};
+
+class USkeletalMeshComponent;
+class UCharacterMovementComponent;
+
+USTRUCT(Experimental)
+struct ANIMATIONWARPINGRUNTIME_API FAnimNode_SlopeWarping : public FAnimNode_SkeletalControlBase
+{
+	GENERATED_USTRUCT_BODY()
+
+	FAnimNode_SlopeWarping();
+
+	USkeletalMeshComponent* MySkelMeshComponent;
+	UCharacterMovementComponent* MyMovementComponent;
+	FAnimInstanceProxy* MyAnimInstanceProxy;
+
+	/** IKFoot Root Bone. **/
+	UPROPERTY(EditAnywhere, Category = "Settings")
+	FBoneReference IKFootRootBone;
+
+	/** Pelvis Bone. **/
+	UPROPERTY(EditAnywhere, Category = "Settings")
+	FBoneReference PelvisBone;
+
+	UPROPERTY(EditAnywhere, Category = "Settings")
+	TArray<FSlopeWarpingFootDefinition> FeetDefinitions;
+
+	UPROPERTY(Transient)
+	TArray<FSlopeWarpingFootData> FeetData;
+
+	UPROPERTY(EditAnywhere, Category = "Settings")
+	FVectorRK4SpringInterpolator PelvisOffsetInterpolator;
+
+	UPROPERTY(EditAnywhere, Category = "Settings", meta = (PinHiddenByDefault))
+	FVector GravityDir;
+
+	UPROPERTY(EditAnywhere, Category = "Settings", meta = (PinHiddenByDefault))
+	FVector CustomFloorOffset;
+
+	UPROPERTY(Transient)
+	float CachedDeltaTime;
+
+	UPROPERTY(Transient)
+	FVector TargetFloorNormalWorldSpace;
+
+	UPROPERTY(EditAnywhere, Category = "Settings")
+	FVectorRK4SpringInterpolator FloorNormalInterpolator;
+
+	UPROPERTY(Transient)
+	FVector TargetFloorOffsetLocalSpace;
+
+	UPROPERTY(EditAnywhere, Category = "Settings")
+	FVectorRK4SpringInterpolator FloorOffsetInterpolator;
+
+	UPROPERTY(EditAnywhere, Category = "Settings")
+	float MaxStepHeight;
+
+	UPROPERTY(EditAnywhere, Category = "Settings")
+	uint32 bKeepMeshInsideOfCapsule : 1;
+
+	UPROPERTY(EditAnywhere, Category = "Settings")
+	uint32 bPullPelvisDown : 1;
+
+	UPROPERTY(EditAnywhere, Category = "Settings", meta = (PinHiddenByDefault))
+	uint32 bUseCustomFloorOffset : 1;
+
+	UPROPERTY(Transient)
+	uint32 bWasOnGround : 1;
+
+	UPROPERTY(Transient)
+	uint32 bShowDebug : 1;
+
+	UPROPERTY(Transient)
+	uint32 bFloorSmoothingInitialized : 1;
+
+	UPROPERTY(Transient)
+	FVector ActorLocation;
+
+protected:
+	UPROPERTY(Transient)
+	FVector GravityDirCompSpace;
+
+public:
+
+	FVector GetFloorNormalFromMovementComponent() const;
+	void SetWorldSpaceTargetFloorInfoFromMovementComponent();
+
+	void GetSmoothedFloorInfo(FCSPose<FCompactPose>& MeshBases, FVector& OutFloorNormal, FVector& OutFloorOffset);
+
+	// FAnimNode_Base interface
+	virtual void GatherDebugData(FNodeDebugData& DebugData) override;
+	virtual void UpdateInternal(const FAnimationUpdateContext& Context) override;
+	// End of FAnimNode_Base interface
+
+	// FAnimNode_SkeletalControlBase interface
+	virtual void Initialize_AnyThread(const FAnimationInitializeContext& Context) override;
+	virtual void EvaluateSkeletalControl_AnyThread(FComponentSpacePoseContext& Output, TArray<FBoneTransform>& OutBoneTransforms) override;
+	virtual bool IsValidToEvaluate(const USkeleton* Skeleton, const FBoneContainer& RequiredBones) override;
+	// End of FAnimNode_SkeletalControlBase interface
+
+private:
+	// FAnimNode_SkeletalControlBase interface
+	virtual void InitializeBoneReferences(const FBoneContainer& RequiredBones) override;
+	// End of FAnimNode_SkeletalControlBase interface
+};
diff --git a/Engine/Plugins/Animation/AnimationWarping/Source/Runtime/Public/BoneControllers/AnimNode_StrideWarping.h b/Engine/Plugins/Animation/AnimationWarping/Source/Runtime/Public/BoneControllers/AnimNode_StrideWarping.h
new file mode 100644
--- /dev/null
+++ b/Engine/Plugins/Animation/AnimationWarping/Source/Runtime/Public/BoneControllers/AnimNode_StrideWarping.h
@@ -0,0 +1,192 @@
+// Copyright Epic Games, Inc. All Rights Reserved.
+
+#pragma once
+
+#include "BoneControllers/BoneControllerTypes.h"
+#include "BoneControllers/BoneControllerSolvers.h"
+#include "BoneControllers/AnimNode_SkeletalControlBase.h"
+#include "AnimNode_StrideWarping.generated.h"
+
+// Foot definition specifying the IK/FK foot bones and Thigh bone
+USTRUCT(BlueprintInternalUseOnly)
+struct ANIMATIONWARPINGRUNTIME_API FStrideWarpingFootDefinition
+{
+	GENERATED_BODY()
+
+	// IK driven foot bone
+	UPROPERTY(EditAnywhere, Category=Settings, meta=(DisplayName="IK Foot Bone"))
+	FBoneReference IKFootBone;
+
+	// FK driven foot bone
+	UPROPERTY(EditAnywhere, Category=Settings, meta=(DisplayName="FK Foot Bone"))
+	FBoneReference FKFootBone;
+
+	// Thigh bone, representing the end of the leg chain prior to reaching the Pelvis Bone 
+	UPROPERTY(EditAnywhere, Category=Settings)
+	FBoneReference ThighBone;
+};
+
+USTRUCT(BlueprintInternalUseOnly)
+struct ANIMATIONWARPINGRUNTIME_API FAnimNode_StrideWarping : public FAnimNode_SkeletalControlBase
+{
+	GENERATED_BODY()
+
+	// Stride warping evaluation mode (Graph or Manual)
+	UPROPERTY(EditAnywhere, Category=Evaluation)
+	EWarpingEvaluationMode Mode = EWarpingEvaluationMode::Manual;
+
+	// Component-space stride direction
+	// Example: A value of <1,0,0> will warp the leg stride along the Forward Vector
+	UPROPERTY(EditAnywhere, BlueprintReadWrite, Category=Evaluation, meta=(PinShownByDefault))
+	FVector StrideDirection = FVector::ForwardVector;
+
+	// Stride scale, specifying the amount of warping applied to the IK foot definitions
+	// Example: A value of 0.5 will decrease the effective leg stride by half, while a value of 2.0 will double it
+	UPROPERTY(EditAnywhere, BlueprintReadWrite, Category=Evaluation, meta=(ClampMin="0.0", PinShownByDefault))
+	float StrideScale = 1.f;
+
+	// Locomotion speed, specifying the current speed of the character
+	// This will be used in the following equation for computing the stride scale: [StrideScale = (LocomotionSpeed / RootMotionSpeed)]
+	// Note: Both speeds should be relative to the delta time of the animation system
+	// 
+	// Stride scale is a value specifying the amount of warping applied to the IK foot definitions
+	// Example: A value of 0.5 will decrease the effective leg stride by half, while a value of 2.0 will double it
+	UPROPERTY(EditAnywhere, BlueprintReadWrite, Category=Evaluation, meta=(ClampMin="0.0", PinShownByDefault))
+	float LocomotionSpeed = 0.f;
+
+	// Minimum locomotion speed required to apply stride warping
+	UPROPERTY(EditAnywhere, BlueprintReadWrite, Category=Evaluation, meta=(ClampMin="0.0", PinHiddenByDefault))
+	float MinLocomotionSpeedThreshold = 10.f;
+
+	// Pevlis Bone definition
+	UPROPERTY(EditAnywhere, Category=Settings)
+	FBoneReference PelvisBone;
+
+	// IK Foot Root Bone definition
+	UPROPERTY(EditAnywhere, Category=Settings, meta=(DisplayName="IK Foot Root Bone"))
+	FBoneReference IKFootRootBone;
+
+	// Foot definitions specifying the IK, FK, and Thigh bone
+	UPROPERTY(EditAnywhere, Category=Settings)
+	TArray<FStrideWarpingFootDefinition> FootDefinitions;
+
+	// Modifies the final stride scale value by optionally clamping and/or interpolating
+	UPROPERTY(EditAnywhere, Category=Settings)
+	FInputClampConstants StrideScaleModifier;
+
+	// Floor normal direction, this value will internally convert into a corresponding Component-space representation prior to warping
+	// Default: World Space, Up Vector: <0,0,1>
+	UPROPERTY(EditAnywhere, BlueprintReadWrite, Category=Advanced, meta=(PinHiddenByDefault))
+	FWarpingVectorValue FloorNormalDirection = { EWarpingVectorMode::WorldSpaceVector, FVector::UpVector };
+
+	// Gravity direction, this value will internally convert into a corresponding Component-space representation prior to warping
+	// Default: World Space, Down Vector: <0,0,-1>
+	UPROPERTY(EditAnywhere, BlueprintReadWrite, Category=Advanced, meta=(PinHiddenByDefault))
+	FWarpingVectorValue GravityDirection = { EWarpingVectorMode::WorldSpaceVector, FVector::DownVector };
+
+	// Solver for controlling how much the pelvis is "pulled down" towards the IK/FK foot definitions during leg limb extension
+	UPROPERTY(EditAnywhere, Category=Advanced, meta=(DisplayName="Pelvis IK Foot Solver"))
+	FIKFootPelvisPullDownSolver PelvisIKFootSolver;
+
+	// Orients the specified (Manual) or computed (Graph) stride direction by the floor normal
+	UPROPERTY(EditAnywhere, Category=Advanced)
+	bool bOrientStrideDirectionUsingFloorNormal = true;
+
+	// Include warping adjustment to the FK thigh bones alongside the IK/FK foot definitions
+	// This is used to help preserve the original overall leg shape
+	UPROPERTY(EditAnywhere, Category=Advanced, meta=(DisplayName="Compensate IK Using FK Thigh Rotation"))
+	bool bCompensateIKUsingFKThighRotation = true;
+
+	// Clamps the IK foot warping to prevent over-extension relative to the overall FK leg
+	UPROPERTY(EditAnywhere, Category=Advanced, meta=(DisplayName="Clamp IK Using FK Limits", EditCondition="bCompensateIKUsingFKThighRotation"))
+	bool bClampIKUsingFKLimits = true;
+
+#if WITH_EDITORONLY_DATA
+	// Scale all debug drawing visualization by a factor
+	UPROPERTY(EditAnywhere, Category=Debug, meta=(ClampMin="0.0"))
+	float DebugDrawScale = 1.f;
+
+	// Enable/Disable stride warping debug drawing
+	UPROPERTY(EditAnywhere, Category=Debug)
+	bool bEnableDebugDraw = false;
+
+	// Enable/Disable IK foot location debug drawing prior to warping
+	UPROPERTY(EditAnywhere, Category=Debug)
+	bool bDebugDrawIKFootOrigin = false;
+
+	// Enable/Disable IK foot location debug drawing following initial foot adjustment
+	UPROPERTY(EditAnywhere, Category=Debug)
+	bool bDebugDrawIKFootAdjustment = false;
+
+	// Enable/Disable pelvis debug drawing following adjustment
+	UPROPERTY(EditAnywhere, Category=Debug)
+	bool bDebugDrawPelvisAdjustment = false;
+
+	// Enable/Disable thigh debug drawing following adjustment
+	UPROPERTY(EditAnywhere, Category=Debug)
+	bool bDebugDrawThighAdjustment = false;
+
+	// Enable/Disable IK foot location debug drawing following all adsjustments (Final warped result)
+	UPROPERTY(EditAnywhere, Category=Debug)
+	bool bDebugDrawIKFootFinal = false;
+#endif
+
+public:
+	// FAnimNode_Base interface
+	virtual void GatherDebugData(FNodeDebugData& DebugData) override;
+	virtual void UpdateInternal(const FAnimationUpdateContext& Context) override;
+	// End of FAnimNode_Base interface
+
+	// FAnimNode_SkeletalControlBase interface
+	virtual void Initialize_AnyThread(const FAnimationInitializeContext& Context) override;
+	virtual void EvaluateSkeletalControl_AnyThread(FComponentSpacePoseContext& Output, TArray<FBoneTransform>& OutBoneTransforms) override;
+	virtual bool IsValidToEvaluate(const USkeleton* Skeleton, const FBoneContainer& RequiredBones) override;
+	// End of FAnimNode_SkeletalControlBase interface
+
+private:
+	// FAnimNode_SkeletalControlBase interface
+	virtual void InitializeBoneReferences(const FBoneContainer& RequiredBones) override;
+	// End of FAnimNode_SkeletalControlBase interface
+
+	struct FStrideWarpingFootData
+	{
+		FCompactPoseBoneIndex IKFootBoneIndex;
+		FCompactPoseBoneIndex FKFootBoneIndex;
+		FCompactPoseBoneIndex ThighBoneIndex;
+		FTransform IKFootBoneTransform;
+
+		FStrideWarpingFootData()
+			: IKFootBoneIndex(INDEX_NONE)
+			, FKFootBoneIndex(INDEX_NONE)
+			, ThighBoneIndex(INDEX_NONE)
+			, IKFootBoneTransform(FTransform::Identity)
+		{
+		}
+	};
+
+	// Internal cached anim instance proxy
+	FAnimInstanceProxy* AnimInstanceProxy = nullptr;
+
+	// Computed IK, FK, Thigh bone indices for the specified foot definitions
+	TArray<FStrideWarpingFootData> FootData;
+
+	// Internal cached stride scale modifier state
+	FInputClampState StrideScaleModifierState;
+
+	// Internal stride direction
+	FVector ActualStrideDirection = FVector::ForwardVector;
+
+	// Internal stride scale
+	float ActualStrideScale = 1.f;
+
+	// Internal cached delta time used for interpolators
+	float CachedDeltaTime = 0.f;
+
+#if WITH_EDITORONLY_DATA
+	// Internal cached debug root motion delta translation
+	FVector CachedRootMotionDeltaTranslation = FVector::ZeroVector;
+
+	// Internal cached debug root motion speed
+	float CachedRootMotionDeltaSpeed = 0.f;
+#endif
+};
diff --git a/Engine/Plugins/Experimental/Animation/AnimationWarping/AnimationWarping.uplugin b/Engine/Plugins/Experimental/Animation/AnimationWarping/AnimationWarping.uplugin
deleted file mode 100644
diff --git a/Engine/Plugins/Experimental/Animation/AnimationWarping/Source/Editor/AnimationWarpingEditor.Build.cs b/Engine/Plugins/Experimental/Animation/AnimationWarping/Source/Editor/AnimationWarpingEditor.Build.cs
deleted file mode 100644
diff --git a/Engine/Plugins/Experimental/Animation/AnimationWarping/Source/Editor/Private/AnimGraph/AnimGraphNode_OrientationWarping.cpp b/Engine/Plugins/Experimental/Animation/AnimationWarping/Source/Editor/Private/AnimGraph/AnimGraphNode_OrientationWarping.cpp
deleted file mode 100644
diff --git a/Engine/Plugins/Experimental/Animation/AnimationWarping/Source/Editor/Private/AnimGraph/AnimGraphNode_SlopeWarping.cpp b/Engine/Plugins/Experimental/Animation/AnimationWarping/Source/Editor/Private/AnimGraph/AnimGraphNode_SlopeWarping.cpp
deleted file mode 100644
diff --git a/Engine/Plugins/Experimental/Animation/AnimationWarping/Source/Editor/Private/AnimGraph/AnimGraphNode_StrideWarping.cpp b/Engine/Plugins/Experimental/Animation/AnimationWarping/Source/Editor/Private/AnimGraph/AnimGraphNode_StrideWarping.cpp
deleted file mode 100644
diff --git a/Engine/Plugins/Experimental/Animation/AnimationWarping/Source/Editor/Private/AnimationWarpingEditorModule.cpp b/Engine/Plugins/Experimental/Animation/AnimationWarping/Source/Editor/Private/AnimationWarpingEditorModule.cpp
deleted file mode 100644
diff --git a/Engine/Plugins/Experimental/Animation/AnimationWarping/Source/Editor/Public/AnimGraph/AnimGraphNode_OrientationWarping.h b/Engine/Plugins/Experimental/Animation/AnimationWarping/Source/Editor/Public/AnimGraph/AnimGraphNode_OrientationWarping.h
deleted file mode 100644
diff --git a/Engine/Plugins/Experimental/Animation/AnimationWarping/Source/Editor/Public/AnimGraph/AnimGraphNode_SlopeWarping.h b/Engine/Plugins/Experimental/Animation/AnimationWarping/Source/Editor/Public/AnimGraph/AnimGraphNode_SlopeWarping.h
deleted file mode 100644
diff --git a/Engine/Plugins/Experimental/Animation/AnimationWarping/Source/Editor/Public/AnimGraph/AnimGraphNode_StrideWarping.h b/Engine/Plugins/Experimental/Animation/AnimationWarping/Source/Editor/Public/AnimGraph/AnimGraphNode_StrideWarping.h
deleted file mode 100644
diff --git a/Engine/Plugins/Experimental/Animation/AnimationWarping/Source/Editor/Public/AnimationWarpingEditorModule.h b/Engine/Plugins/Experimental/Animation/AnimationWarping/Source/Editor/Public/AnimationWarpingEditorModule.h
deleted file mode 100644
diff --git a/Engine/Plugins/Experimental/Animation/AnimationWarping/Source/Runtime/AnimationWarpingRuntime.Build.cs b/Engine/Plugins/Experimental/Animation/AnimationWarping/Source/Runtime/AnimationWarpingRuntime.Build.cs
deleted file mode 100644
diff --git a/Engine/Plugins/Experimental/Animation/AnimationWarping/Source/Runtime/Private/AnimationWarpingRuntimeModule.cpp b/Engine/Plugins/Experimental/Animation/AnimationWarping/Source/Runtime/Private/AnimationWarpingRuntimeModule.cpp
deleted file mode 100644
diff --git a/Engine/Plugins/Experimental/Animation/AnimationWarping/Source/Runtime/Private/BoneControllers/AnimNode_OrientationWarping.cpp b/Engine/Plugins/Experimental/Animation/AnimationWarping/Source/Runtime/Private/BoneControllers/AnimNode_OrientationWarping.cpp
deleted file mode 100644
diff --git a/Engine/Plugins/Experimental/Animation/AnimationWarping/Source/Runtime/Private/BoneControllers/AnimNode_SlopeWarping.cpp b/Engine/Plugins/Experimental/Animation/AnimationWarping/Source/Runtime/Private/BoneControllers/AnimNode_SlopeWarping.cpp
deleted file mode 100644
diff --git a/Engine/Plugins/Experimental/Animation/AnimationWarping/Source/Runtime/Private/BoneControllers/AnimNode_StrideWarping.cpp b/Engine/Plugins/Experimental/Animation/AnimationWarping/Source/Runtime/Private/BoneControllers/AnimNode_StrideWarping.cpp
deleted file mode 100644
diff --git a/Engine/Plugins/Experimental/Animation/AnimationWarping/Source/Runtime/Public/AnimationWarpingRuntimeModule.h b/Engine/Plugins/Experimental/Animation/AnimationWarping/Source/Runtime/Public/AnimationWarpingRuntimeModule.h
deleted file mode 100644
diff --git a/Engine/Plugins/Experimental/Animation/AnimationWarping/Source/Runtime/Public/BoneControllers/AnimNode_OrientationWarping.h b/Engine/Plugins/Experimental/Animation/AnimationWarping/Source/Runtime/Public/BoneControllers/AnimNode_OrientationWarping.h
deleted file mode 100644
diff --git a/Engine/Plugins/Experimental/Animation/AnimationWarping/Source/Runtime/Public/BoneControllers/AnimNode_SlopeWarping.h b/Engine/Plugins/Experimental/Animation/AnimationWarping/Source/Runtime/Public/BoneControllers/AnimNode_SlopeWarping.h
deleted file mode 100644
diff --git a/Engine/Plugins/Experimental/Animation/AnimationWarping/Source/Runtime/Public/BoneControllers/AnimNode_StrideWarping.h b/Engine/Plugins/Experimental/Animation/AnimationWarping/Source/Runtime/Public/BoneControllers/AnimNode_StrideWarping.h
deleted file mode 100644
