#trivial

Animation Warping Editor Module type fix

diff --git a/Engine/Plugins/Animation/AnimationWarping/AnimationWarping.uplugin b/Engine/Plugins/Animation/AnimationWarping/AnimationWarping.uplugin
--- a/Engine/Plugins/Animation/AnimationWarping/AnimationWarping.uplugin
+++ b/Engine/Plugins/Animation/AnimationWarping/AnimationWarping.uplugin
@@ -20,7 +20,7 @@
 		},
 		{
 			"Name": "AnimationWarpingEditor",
-			"Type": "Editor",
+			"Type": "UncookedOnly",
 			"LoadingPhase": "Default"
 		},
 	]
