#trivial

Animation Warping tooltip/comment quality improvements and minor code cleanup.

diff --git a/Engine/Plugins/Animation/AnimationWarping/Source/Editor/Private/AnimGraph/AnimGraphNode_OrientationWarping.cpp b/Engine/Plugins/Animation/AnimationWarping/Source/Editor/Private/AnimGraph/AnimGraphNode_OrientationWarping.cpp
--- a/Engine/Plugins/Animation/AnimationWarping/Source/Editor/Private/AnimGraph/AnimGraphNode_OrientationWarping.cpp
+++ b/Engine/Plugins/Animation/AnimationWarping/Source/Editor/Private/AnimGraph/AnimGraphNode_OrientationWarping.cpp
@@ -21,7 +21,7 @@ FText UAnimGraphNode_OrientationWarping::GetControllerDescription() const
 
 FText UAnimGraphNode_OrientationWarping::GetTooltipText() const
 {
-	return LOCTEXT("OrientationWarpingTooltip", "Warps the root bone, and counter rotates the body, to match locomotion direction.");
+	return LOCTEXT("OrientationWarpingTooltip", "Rotates the root and lower body by the specified angle, while counter rotating the upper body to maintain the forward facing direction.");
 }
 
 FText UAnimGraphNode_OrientationWarping::GetNodeTitle(ENodeTitleType::Type TitleType) const
diff --git a/Engine/Plugins/Animation/AnimationWarping/Source/Editor/Private/AnimGraph/AnimGraphNode_StrideWarping.cpp b/Engine/Plugins/Animation/AnimationWarping/Source/Editor/Private/AnimGraph/AnimGraphNode_StrideWarping.cpp
--- a/Engine/Plugins/Animation/AnimationWarping/Source/Editor/Private/AnimGraph/AnimGraphNode_StrideWarping.cpp
+++ b/Engine/Plugins/Animation/AnimationWarping/Source/Editor/Private/AnimGraph/AnimGraphNode_StrideWarping.cpp
@@ -21,7 +21,7 @@ FText UAnimGraphNode_StrideWarping::GetControllerDescription() const
 
 FText UAnimGraphNode_StrideWarping::GetTooltipText() const
 {
-	return LOCTEXT("StrideWarpingTooltip", "Warps the pelvis and feet to match the locomotion speed.");
+	return LOCTEXT("StrideWarpingTooltip", "Adjusts the speed of the animation by warping the leg stride with the specified scale and direction.");
 }
 
 FText UAnimGraphNode_StrideWarping::GetNodeTitle(ENodeTitleType::Type TitleType) const
diff --git a/Engine/Plugins/Animation/AnimationWarping/Source/Runtime/Private/BoneControllers/AnimNode_OrientationWarping.cpp b/Engine/Plugins/Animation/AnimationWarping/Source/Runtime/Private/BoneControllers/AnimNode_OrientationWarping.cpp
--- a/Engine/Plugins/Animation/AnimationWarping/Source/Runtime/Private/BoneControllers/AnimNode_OrientationWarping.cpp
+++ b/Engine/Plugins/Animation/AnimationWarping/Source/Runtime/Private/BoneControllers/AnimNode_OrientationWarping.cpp
@@ -161,12 +161,11 @@ void FAnimNode_OrientationWarping::EvaluateSkeletalControl_AnyThread(FComponentS
 			}
 
 			// For interpolated warping, guarantee that PreviousOrientationAngle is respect to the current frame's root motion direction 
-			const FVector RootMotionCrossProduct = RootMotionDeltaDirection.Cross(PreviousRootMotionDeltaDirection);
-			float RootMotionDeltaAngle = FMath::Acos(RootMotionDeltaDirection.Dot(PreviousRootMotionDeltaDirection));
-			RootMotionDeltaAngle *= FMath::Sign(RotationAxisVector.Dot(RootMotionCrossProduct));
+			float RootMotionDeltaAngleDifference = FMath::Acos(RootMotionDeltaDirection.Dot(PreviousRootMotionDeltaDirection));
+			RootMotionDeltaAngleDifference *= FMath::Sign(RotationAxisVector.Dot(RootMotionDeltaDirection.Cross(PreviousRootMotionDeltaDirection)));
 
 			PreviousRootMotionDeltaDirection = RootMotionDeltaDirection;
-			PreviousOrientationAngle += RootMotionDeltaAngle;
+			PreviousOrientationAngle += RootMotionDeltaAngleDifference;
 
 			// Rotate the root motion delta fully by the warped angle
 			const FVector WarpedRootMotionTranslationDelta = WarpedRotation.RotateVector(RootMotionDeltaTranslation);
diff --git a/Engine/Plugins/Animation/AnimationWarping/Source/Runtime/Public/BoneControllers/AnimNode_OrientationWarping.h b/Engine/Plugins/Animation/AnimationWarping/Source/Runtime/Public/BoneControllers/AnimNode_OrientationWarping.h
--- a/Engine/Plugins/Animation/AnimationWarping/Source/Runtime/Public/BoneControllers/AnimNode_OrientationWarping.h
+++ b/Engine/Plugins/Animation/AnimationWarping/Source/Runtime/Public/BoneControllers/AnimNode_OrientationWarping.h
@@ -20,14 +20,14 @@ struct ANIMATIONWARPINGRUNTIME_API FAnimNode_OrientationWarping : public FAnimNo
 	float OrientationAngle = 0.f;
 
 	// The character locomotion angle (in degrees) relative to the specified RotationAxis
-	// This will be used in the following equation to compute the effective orientation warping angle: [Orientation = RotationBetween(RootMotionDirection, LocomotionDirection)]
+	// This will be used in the following equation for computing the orientation angle: [Orientation = RotationBetween(RootMotionDirection, LocomotionDirection)]
 	UPROPERTY(EditAnywhere, BlueprintReadWrite, Category=Evaluation, meta=(PinShownByDefault))
 	float LocomotionAngle = 0.f;
 
 	// Specifies an angle threshold to prevent erroneous over-rotation of the character, disabled with a value of 0
 	//
 	// When the effective orientation warping angle is detected to be greater than this value (default: 90 degrees) the locomotion direction will be inverted prior to warping
-	// This will result in the following equation: [Orientation = RotationBetween(RootMotionDirection, -LocomotionDirection)]
+	// This will be used in the following equation: [Orientation = RotationBetween(RootMotionDirection, -LocomotionDirection)]
 	//
 	// Example: Playing a forward running animation while the motion is going backward 
 	// Rather than orientation warping by 180 degrees, the system will warp by 0 degrees 
@@ -35,7 +35,7 @@ struct ANIMATIONWARPINGRUNTIME_API FAnimNode_OrientationWarping : public FAnimNo
 	float LocomotionAngleDeltaThreshold = 90.f;
 
 	// Spine bone definitions
-	// Used to counter the rotation of the body to keep the character facing towards the specified angle (Orientation Angle, Locomotion Angle)
+	// Used to counter rotate the body in order to keep the character facing forward
 	// The amount of counter rotation applied is driven by DistributedBoneOrientationAlpha
 	UPROPERTY(EditAnywhere, Category=Settings)
 	TArray<FBoneReference> SpineBones;
@@ -48,16 +48,16 @@ struct ANIMATIONWARPINGRUNTIME_API FAnimNode_OrientationWarping : public FAnimNo
 	UPROPERTY(EditAnywhere, Category=Settings, meta=(DisplayName="IK Foot Bones"))
 	TArray<FBoneReference> IKFootBones;
 
-	// Specifies how much rotation is applied to the character body versus IK feet
-	UPROPERTY(EditAnywhere, Category=Settings, meta=(ClampMin="0.0", ClampMax="1.0"))
-	float DistributedBoneOrientationAlpha = 0.5f;
-
 	// Rotation axis used when rotating the character body
 	UPROPERTY(EditAnywhere, Category=Settings)
 	TEnumAsByte<EAxis::Type> RotationAxis = EAxis::Z;
 
+	// Specifies how much rotation is applied to the character body versus IK feet
+	UPROPERTY(EditAnywhere, Category=Settings, meta=(ClampMin="0.0", ClampMax="1.0"))
+	float DistributedBoneOrientationAlpha = 0.5f;
+
 	// Specifies the interpolation speed (in Alpha per second) towards reaching the final warped rotation angle
-	// A value of 0 will cause instantaneous warping, while a greater value will introduce smoothing
+	// A value of 0 will cause instantaneous rotation, while a greater value will introduce smoothing
 	UPROPERTY(EditAnywhere, Category=Settings, meta=(ClampMin="0.0"))
 	float RotationInterpSpeed = 10.f;
 
diff --git a/Engine/Plugins/Animation/AnimationWarping/Source/Runtime/Public/BoneControllers/AnimNode_StrideWarping.h b/Engine/Plugins/Animation/AnimationWarping/Source/Runtime/Public/BoneControllers/AnimNode_StrideWarping.h
--- a/Engine/Plugins/Animation/AnimationWarping/Source/Runtime/Public/BoneControllers/AnimNode_StrideWarping.h
+++ b/Engine/Plugins/Animation/AnimationWarping/Source/Runtime/Public/BoneControllers/AnimNode_StrideWarping.h
@@ -40,14 +40,14 @@ struct ANIMATIONWARPINGRUNTIME_API FAnimNode_StrideWarping : public FAnimNode_Sk
 	UPROPERTY(EditAnywhere, BlueprintReadWrite, Category=Evaluation, meta=(PinShownByDefault))
 	FVector StrideDirection = FVector::ForwardVector;
 
-	// Stride scale, specifying the amount of warping applied to the IK foot definitions
+	// Stride scale, specifying the amount of warping applied to the foot definitions
 	// Example: A value of 0.5 will decrease the effective leg stride by half, while a value of 2.0 will double it
 	UPROPERTY(EditAnywhere, BlueprintReadWrite, Category=Evaluation, meta=(ClampMin="0.0", PinShownByDefault))
 	float StrideScale = 1.f;
 
 	// Locomotion speed, specifying the current speed of the character
 	// This will be used in the following equation for computing the stride scale: [StrideScale = (LocomotionSpeed / RootMotionSpeed)]
-	// Note: Both speeds should be relative to the delta time of the animation system
+	// Note: This speed should be relative to the delta time of the animation graph
 	// 
 	// Stride scale is a value specifying the amount of warping applied to the IK foot definitions
 	// Example: A value of 0.5 will decrease the effective leg stride by half, while a value of 2.0 will double it
diff --git a/Engine/Source/Runtime/AnimGraphRuntime/Public/BoneControllers/BoneControllerTypes.h b/Engine/Source/Runtime/AnimGraphRuntime/Public/BoneControllers/BoneControllerTypes.h
--- a/Engine/Source/Runtime/AnimGraphRuntime/Public/BoneControllers/BoneControllerTypes.h
+++ b/Engine/Source/Runtime/AnimGraphRuntime/Public/BoneControllers/BoneControllerTypes.h
@@ -8,14 +8,15 @@
 
 struct FAnimInstanceProxy;
 
+// Specifies the evaluation mode of an animation warping node
 UENUM(BlueprintInternalUseOnly)
 enum class EWarpingEvaluationMode : uint8
 {
-	// Pose warping node parameters are entirely driven by the user
+	// Animation warping evaluation parameters are driven by user settings.
 	Manual,
-	// Pose warping nodes may participate in being graph-driven. This means some
-	// properties of the warp may be automatically configured by the accumulated 
-	// root motion delta contribution of the animation graph leading into the node
+	// Animation warping evaluation parameters are graph-driven. This means some
+	// properties of the node are automatically computed using the accumulated 
+	// root motion delta contribution of the animation graph leading into it.
 	Graph
 };
 
@@ -33,6 +34,7 @@ enum class EWarpingVectorMode : uint8
 	IKFootRootLocalSpaceVector,
 };
 
+// Vector values which may be specified in a configured space
 USTRUCT(BlueprintType)
 struct ANIMGRAPHRUNTIME_API FWarpingVectorValue
 {
