Fix for crash in AnimNode_StrideWarping

diff --git a/Engine/Plugins/Animation/AnimationWarping/Source/Runtime/Private/BoneControllers/AnimNode_StrideWarping.cpp b/Engine/Plugins/Animation/AnimationWarping/Source/Runtime/Private/BoneControllers/AnimNode_StrideWarping.cpp
--- a/Engine/Plugins/Animation/AnimationWarping/Source/Runtime/Private/BoneControllers/AnimNode_StrideWarping.cpp
+++ b/Engine/Plugins/Animation/AnimationWarping/Source/Runtime/Private/BoneControllers/AnimNode_StrideWarping.cpp
@@ -156,7 +156,8 @@ void FAnimNode_StrideWarping::EvaluateSkeletalControl_AnyThread(FComponentSpaceP
 #endif
 		// Graph driven stride scale factor will be determined by the ratio of the
 		// locomotion (capsule/physics) speed against the animation root motion speed
-		ActualStrideScale = LocomotionSpeed / CachedRootMotionDeltaSpeed;
+		// If there's (almost) no extracted root motion, we don't want impossibly large strides.
+		ActualStrideScale = FMath::IsNearlyEqual(CachedRootMotionDeltaSpeed, 0.f, KINDA_SMALL_NUMBER) ? 1.0f : LocomotionSpeed / CachedRootMotionDeltaSpeed;
 	}
 
 	// Allow the opportunity for stride scale clamping and interpolation regardless of evaluation mode
