Removed ensures for missing root motion delta attribute from stride/orientating warping nodes, and added it to on-screen debug data
-review-18678170 [REDACTED_HANDLE]

diff --git a/Engine/Plugins/Animation/AnimationWarping/Source/Runtime/Private/BoneControllers/AnimNode_OrientationWarping.cpp b/Engine/Plugins/Animation/AnimationWarping/Source/Runtime/Private/BoneControllers/AnimNode_OrientationWarping.cpp
--- a/Engine/Plugins/Animation/AnimationWarping/Source/Runtime/Private/BoneControllers/AnimNode_OrientationWarping.cpp
+++ b/Engine/Plugins/Animation/AnimationWarping/Source/Runtime/Private/BoneControllers/AnimNode_OrientationWarping.cpp
@@ -51,6 +51,9 @@ void FAnimNode_OrientationWarping::GatherDebugData(FNodeDebugData& DebugData)
 			DebugLine += FString::Printf(TEXT("\n - Orientation Angle: (%.3fd)"), FMath::RadiansToDegrees(ActualOrientationAngle));
 			DebugLine += FString::Printf(TEXT("\n - Locomotion Angle: (%.3fd)"), FMath::RadiansToDegrees(LocomotionAngle));
 			DebugLine += FString::Printf(TEXT("\n - Locomotion Delta Angle Threshold: (%.3fd)"), LocomotionAngleDeltaThreshold);
+#if WITH_EDITORONLY_DATA
+			DebugLine += FString::Printf(TEXT("\n - Root Motion Delta Attribute Found: %s)"), (bFoundRootMotionAttribute) ? TEXT("true") : TEXT("false"));
+#endif
 		}
 		DebugLine += FString::Printf(TEXT("\n - Distributed Bone Orientation Alpha: (%.3fd)"), DistributedBoneOrientationAlpha);
 		if (const UEnum* TypeEnum = FindObject<UEnum>(ANY_PACKAGE, TEXT("EAxis")))
@@ -102,6 +105,10 @@ void FAnimNode_OrientationWarping::EvaluateSkeletalControl_AnyThread(FComponentS
 		ensureMsgf(bGraphDrivenWarping, TEXT("Graph driven Orientation Warping expected a valid root motion delta provider interface."));
 	}
 
+#if WITH_EDITORONLY_DATA
+	bFoundRootMotionAttribute = false;
+#endif
+
 	// We will likely need to revisit LocomotionAngle participating as an input to orientation warping.
 	// Without velocity information from the motion model (such as the capsule), LocomotionAngle isn't enough
 	// information in isolation for all cases when deciding to warp.
@@ -122,10 +129,13 @@ void FAnimNode_OrientationWarping::EvaluateSkeletalControl_AnyThread(FComponentS
 		bGraphDrivenWarping = RootMotionProvider->ExtractRootMotion(Output.CustomAttributes, RootMotionTransformDelta);
 
 		// Graph driven orientation warping will modify the incoming root motion to orient towards the intended locomotion angle
-		if (ensureMsgf(bGraphDrivenWarping, TEXT(
-			"Graph driven Orientation Warping expected a root motion delta to be present in the attribute stream.\n"
-			"Check that the evaluating Animation Sequence(s) have: \"Enable Root Motion\" enabled.")))
+		if (bGraphDrivenWarping)
 		{
+#if WITH_EDITORONLY_DATA
+			// Graph driven Orientation Warping expects a root motion delta to be present in the attribute stream.
+			bFoundRootMotionAttribute = true;
+#endif
+
 			// In UE, forward is defined as +x; consequently this is also true when sampling an actor's velocity. Historically the skeletal 
 			// mesh component forward will not match the actor, requiring us to correct the rotation before sampling the LocomotionForward.
 			// In order to make orientation warping 'pure' in the future we will need to provide more context about the intent of
@@ -177,7 +187,7 @@ void FAnimNode_OrientationWarping::EvaluateSkeletalControl_AnyThread(FComponentS
 		}
 		else
 		{
-			// Early exit on broken graph driven behavior 
+			// Early exit on missing root motion delta attribute
 			return;
 		}
 	} 
diff --git a/Engine/Plugins/Animation/AnimationWarping/Source/Runtime/Private/BoneControllers/AnimNode_StrideWarping.cpp b/Engine/Plugins/Animation/AnimationWarping/Source/Runtime/Private/BoneControllers/AnimNode_StrideWarping.cpp
--- a/Engine/Plugins/Animation/AnimationWarping/Source/Runtime/Private/BoneControllers/AnimNode_StrideWarping.cpp
+++ b/Engine/Plugins/Animation/AnimationWarping/Source/Runtime/Private/BoneControllers/AnimNode_StrideWarping.cpp
@@ -25,6 +25,7 @@ void FAnimNode_StrideWarping::GatherDebugData(FNodeDebugData& DebugData)
 		{
 			DebugLine += FString::Printf(TEXT("\n - Locomotion Speed: (%.3fd)"), LocomotionSpeed);
 #if WITH_EDITORONLY_DATA
+			DebugLine += FString::Printf(TEXT("\n - Root Motion Delta Attribute Found: %s)"), (bFoundRootMotionAttribute) ? TEXT("true") : TEXT("false"));
 			DebugLine += FString::Printf(TEXT("\n - Root Motion Direction: (%s)"), *(CachedRootMotionDeltaTranslation.GetSafeNormal().ToCompactString()));
 			DebugLine += FString::Printf(TEXT("\n - Root Motion Speed: (%.3fd)"), CachedRootMotionDeltaSpeed);
 #endif
@@ -80,23 +81,27 @@ void FAnimNode_StrideWarping::EvaluateSkeletalControl_AnyThread(FComponentSpaceP
 	}
 
 	FTransform RootMotionTransformDelta = FTransform::Identity;
-#if !WITH_EDITORONLY_DATA
+#if WITH_EDITORONLY_DATA
+	bFoundRootMotionAttribute = false;
+#else
 	FVector CachedRootMotionDeltaTranslation = FVector::ZeroVector;
 #endif
 	if (bGraphDrivenWarping)
 	{
 		// Graph driven stride warping will override the manual stride direction with the intent of the current animation sub-graph's accumulated root motion
 		bGraphDrivenWarping = RootMotionProvider->ExtractRootMotion(Output.CustomAttributes, RootMotionTransformDelta);
-		if (ensureMsgf(bGraphDrivenWarping, TEXT(
-			"Graph driven Stride Warping expected a root motion delta to be present in the attribute stream.\n"
-			"Check that the evaluating Animation Sequence(s) have: \"Enable Root Motion\" enabled.")))
+		if (bGraphDrivenWarping)
 		{
 			CachedRootMotionDeltaTranslation = RootMotionTransformDelta.GetTranslation();
 			ActualStrideDirection = CachedRootMotionDeltaTranslation.GetSafeNormal();
+#if WITH_EDITORONLY_DATA
+			// Graph driven Stride Warping expects a root motion delta to be present in the attribute stream.
+			bFoundRootMotionAttribute = true;
+#endif
 		}
 		else
 		{
-			// Early exit on broken graph driven behavior 
+			// Early exit on missing root motion delta attribute
 			return;
 		}
 	}
diff --git a/Engine/Plugins/Animation/AnimationWarping/Source/Runtime/Public/BoneControllers/AnimNode_OrientationWarping.h b/Engine/Plugins/Animation/AnimationWarping/Source/Runtime/Public/BoneControllers/AnimNode_OrientationWarping.h
--- a/Engine/Plugins/Animation/AnimationWarping/Source/Runtime/Public/BoneControllers/AnimNode_OrientationWarping.h
+++ b/Engine/Plugins/Animation/AnimationWarping/Source/Runtime/Public/BoneControllers/AnimNode_OrientationWarping.h
@@ -141,4 +141,9 @@ private:
 	
 	// Internal orientation warping angle
 	float ActualOrientationAngle = 0.f;
+
+#if WITH_EDITORONLY_DATA
+	// Whether we found a root motion delta attribute in the attribute stream on graph driven mode
+	bool bFoundRootMotionAttribute = false;
+#endif
 };
\ No newline at end of file
diff --git a/Engine/Plugins/Animation/AnimationWarping/Source/Runtime/Public/BoneControllers/AnimNode_StrideWarping.h b/Engine/Plugins/Animation/AnimationWarping/Source/Runtime/Public/BoneControllers/AnimNode_StrideWarping.h
--- a/Engine/Plugins/Animation/AnimationWarping/Source/Runtime/Public/BoneControllers/AnimNode_StrideWarping.h
+++ b/Engine/Plugins/Animation/AnimationWarping/Source/Runtime/Public/BoneControllers/AnimNode_StrideWarping.h
@@ -185,8 +185,11 @@ private:
 #if WITH_EDITORONLY_DATA
 	// Internal cached debug root motion delta translation
 	FVector CachedRootMotionDeltaTranslation = FVector::ZeroVector;
-
+	
 	// Internal cached debug root motion speed
 	float CachedRootMotionDeltaSpeed = 0.f;
+
+	// Whether we found a root motion delta attribute in the attribute stream on graph driven mode
+	bool bFoundRootMotionAttribute = false;
 #endif
 };
