Added experimental foot placement node.

diff --git a/Engine/Plugins/Animation/AnimationWarping/Source/Editor/Private/AnimGraph/AnimGraphNode_FootPlacement.cpp b/Engine/Plugins/Animation/AnimationWarping/Source/Editor/Private/AnimGraph/AnimGraphNode_FootPlacement.cpp
new file mode 100644
--- /dev/null
+++ b/Engine/Plugins/Animation/AnimationWarping/Source/Editor/Private/AnimGraph/AnimGraphNode_FootPlacement.cpp
@@ -0,0 +1,33 @@
+// Copyright Epic Games, Inc. All Rights Reserved.
+
+#include "AnimGraph/AnimGraphNode_FootPlacement.h"
+#include "Animation/AnimRootMotionProvider.h"
+#include "Kismet2/BlueprintEditorUtils.h"
+#include "DetailLayoutBuilder.h"
+
+/////////////////////////////////////////////////////
+// UAnimGraphNode_FootPlacement
+
+#define LOCTEXT_NAMESPACE "AnimGraphNode_FootPlacement"
+
+UAnimGraphNode_FootPlacement::UAnimGraphNode_FootPlacement(const FObjectInitializer& ObjectInitializer)
+	: Super(ObjectInitializer)
+{
+}
+
+FText UAnimGraphNode_FootPlacement::GetControllerDescription() const
+{
+	return LOCTEXT("FootPlacement", "Foot Placement");
+}
+
+FText UAnimGraphNode_FootPlacement::GetTooltipText() const
+{
+	return LOCTEXT("FootPlacementTooltip", "Foot Placement.");
+}
+
+FText UAnimGraphNode_FootPlacement::GetNodeTitle(ENodeTitleType::Type TitleType) const
+{
+	return GetControllerDescription();
+}
+
+#undef LOCTEXT_NAMESPACE
diff --git a/Engine/Plugins/Animation/AnimationWarping/Source/Editor/Public/AnimGraph/AnimGraphNode_FootPlacement.h b/Engine/Plugins/Animation/AnimationWarping/Source/Editor/Public/AnimGraph/AnimGraphNode_FootPlacement.h
new file mode 100644
--- /dev/null
+++ b/Engine/Plugins/Animation/AnimationWarping/Source/Editor/Public/AnimGraph/AnimGraphNode_FootPlacement.h
@@ -0,0 +1,27 @@
+// Copyright Epic Games, Inc. All Rights Reserved.
+
+#pragma once
+#include "AnimGraphNode_SkeletalControlBase.h"
+#include "BoneControllers/AnimNode_FootPlacement.h"
+#include "AnimGraphNode_FootPlacement.generated.h"
+
+UCLASS(Experimental)
+class ANIMATIONWARPINGEDITOR_API UAnimGraphNode_FootPlacement : public UAnimGraphNode_SkeletalControlBase
+{
+	GENERATED_UCLASS_BODY()
+
+	UPROPERTY(EditAnywhere, Category = Settings)
+	FAnimNode_FootPlacement Node;
+
+public:
+	// UEdGraphNode interface
+	virtual FText GetNodeTitle(ENodeTitleType::Type TitleType) const override;
+	virtual FText GetTooltipText() const override;
+	// End of UEdGraphNode interface
+
+protected:
+	// UAnimGraphNode_SkeletalControlBase interface
+	virtual FText GetControllerDescription() const override;
+	virtual const FAnimNode_SkeletalControlBase* GetNode() const override { return &Node; }
+	// End of UAnimGraphNode_SkeletalControlBase interface
+};
diff --git a/Engine/Plugins/Animation/AnimationWarping/Source/Runtime/Private/BoneControllers/AnimNode_FootPlacement.cpp b/Engine/Plugins/Animation/AnimationWarping/Source/Runtime/Private/BoneControllers/AnimNode_FootPlacement.cpp
new file mode 100644
--- /dev/null
+++ b/Engine/Plugins/Animation/AnimationWarping/Source/Runtime/Private/BoneControllers/AnimNode_FootPlacement.cpp
@@ -0,0 +1,1453 @@
+// Copyright Epic Games, Inc. All Rights Reserved.
+
+#include "BoneControllers/AnimNode_FootPlacement.h"
+#include "AnimationRuntime.h"
+#include "Animation/AnimInstanceProxy.h"
+#include "GameFramework/Character.h"
+#include "GameFramework/CharacterMovementComponent.h"
+#include "Components/CapsuleComponent.h"
+
+DECLARE_CYCLE_STAT(TEXT("Foot Placement Eval"), STAT_FootPlacement_Eval, STATGROUP_Anim);
+
+#if ENABLE_ANIM_DEBUG
+static TAutoConsoleVariable<bool> CVarAnimNodeFootPlacementEnable(TEXT("a.AnimNode.FootPlacement.Enable"), true, TEXT("Enable/Disable Foot Placement"));
+static TAutoConsoleVariable<bool> CVarAnimNodeFootPlacementEnableLock(TEXT("a.AnimNode.FootPlacement.Enable.Lock"), true, TEXT("Enable/Disable Foot Locking"));
+static TAutoConsoleVariable<bool> CVarAnimNodeFootPlacementDebug(TEXT("a.AnimNode.FootPlacement.Debug"), false, TEXT("Turn on visualization debugging for Foot Placement"));
+static TAutoConsoleVariable<int> CVarAnimNodeFootPlacementDebugDrawHistory(TEXT("a.AnimNode.FootPlacement.Debug.DrawHistory"), 0,
+	TEXT("Turn on history visualization debugging 0 = Disabled, -1 = Pelvis, >1 = Foot Index. Clear with FlushPersistentDebugLines"));
+#endif
+
+namespace UE::Anim::FootPlacement
+{
+	////////////////////////////////////////////////////////////////////////////////////////////////////////////////////////
+	/// UE::Anim::FootPlacement::FEvaluationContext
+	struct FEvaluationContext
+	{
+		FEvaluationContext(
+			FComponentSpacePoseContext& InCSPContext,
+			const FVector& InApproachDirWS,
+			const float InUpdateDeltaTime);
+
+		FComponentSpacePoseContext& CSPContext;
+
+		//TODO: maybe store character?
+		class AActor* OwningActor = nullptr;
+		class UWorld* World = nullptr;
+		class UCharacterMovementComponent* MovementComponent = nullptr;
+		FTransform OwningComponentToWorld = FTransform::Identity;
+		float UpdateDeltaTime = 0.0f;
+		FVector ApproachDirWS = -FVector::UpVector;
+		FVector ApproachDirCS = -FVector::UpVector;
+
+		FVector GetMovementComponentFloorNormal() const
+		{
+			if (!MovementComponent)
+			{
+				return -ApproachDirWS;
+			}
+
+			return (MovementComponent->CurrentFloor.bBlockingHit) ?
+				MovementComponent->CurrentFloor.HitResult.ImpactNormal : -ApproachDirWS;
+		}
+
+		bool GetMovementComponentIsWalkable(const FHitResult& InHit) const
+		{
+			if (!MovementComponent)
+			{
+				return false;
+			}
+
+			return MovementComponent->IsWalkable(InHit);
+		}
+	};
+
+	FEvaluationContext::FEvaluationContext(
+		FComponentSpacePoseContext& InCSPContext,
+		const FVector& InApproachDirCS,
+		const float InUpdateDeltaTime)
+		: CSPContext(InCSPContext)
+		, UpdateDeltaTime(InUpdateDeltaTime)
+		, ApproachDirCS(InApproachDirCS)
+	{
+		const USkeletalMeshComponent* OwningComponent = CSPContext.AnimInstanceProxy->GetSkelMeshComponent();
+		OwningActor = OwningComponent->GetOwner();
+		World = OwningComponent->GetWorld();
+
+		ACharacter* CharacterOwner = Cast<ACharacter>(OwningActor);
+		MovementComponent = CharacterOwner ? CharacterOwner->GetCharacterMovement() : nullptr;
+		OwningComponentToWorld = OwningComponent->GetComponentToWorld();
+
+		ApproachDirWS = OwningComponentToWorld.TransformVector(ApproachDirCS);
+	}
+
+	static FVector ReOrientNormal(const FVector& ApproachDir, const FVector& InNormal, FVector& PointA, const FVector& PointB)
+	{
+		const FVector AxisX = (PointA - PointB).GetSafeNormal();
+		if (!AxisX.IsNearlyZero() && !InNormal.IsNearlyZero() && (FMath::Abs(AxisX | InNormal) > DELTA))
+		{
+			const FVector AxisY = AxisX ^ InNormal;
+			const FVector AxisZ = AxisX ^ AxisY;
+
+			// Make sure our normal points upwards. (take into account gravity dir?)
+			return ((AxisZ | -ApproachDir) > 0.f) ? AxisZ : -AxisZ;
+		}
+
+		return InNormal;
+	}
+
+	static bool FindPlantTraceImpact(
+		const FEvaluationContext& Context,
+		const FFootPlacementTraceSettings& TraceSettings,
+		const bool bCheckComplex,
+		const FVector& StartPositionWS,
+		FVector& OutImpactLocationWS,
+		FVector& OutImpactNormalWS)
+	{
+		OutImpactLocationWS = Context.OwningComponentToWorld.GetLocation();
+		OutImpactNormalWS = Context.OwningComponentToWorld.GetRotation().GetUpVector();
+
+		if (!IsValid(Context.World) || !TraceSettings.bEnabled)
+		{
+			return false;
+		}
+
+		const FCollisionShape CollisionShape = FCollisionShape::MakeSphere(TraceSettings.SweepRadius);
+
+		const FVector TraceDirectionWS = Context.ApproachDirWS;
+		const FVector TraceStart = StartPositionWS + (TraceDirectionWS * TraceSettings.StartOffset);
+		const FVector TraceEnd = StartPositionWS + (TraceSettings.EndOffset * TraceDirectionWS);
+
+		FCollisionQueryParams QueryParams;
+		QueryParams.bTraceComplex = bCheckComplex;
+		// Ignore self and all attached components
+		QueryParams.AddIgnoredActor(Context.OwningActor);
+
+		const ECollisionChannel CollisionChannel = UEngineTypes::ConvertToCollisionChannel(bCheckComplex ? TraceSettings.ComplexTraceChannel : TraceSettings.SimpleTraceChannel);
+
+		FHitResult HitResult;
+		const bool bHit = Context.World->SweepSingleByChannel(
+			HitResult, TraceStart, TraceEnd, FQuat::Identity, CollisionChannel, CollisionShape, QueryParams);
+
+		if (!Context.GetMovementComponentIsWalkable(HitResult))
+		{
+			return false;
+		}
+
+		OutImpactLocationWS = HitResult.ImpactPoint;
+		OutImpactNormalWS = HitResult.ImpactNormal;
+
+		return true;
+	}
+
+
+	static bool FindPlantPlane(const FEvaluationContext& Context,
+		const FFootPlacementTraceSettings& TraceSettings,
+		const FVector& StartPositionWS,
+		const bool bCheckComplex,
+		FPlane& OutPlantPlaneWS,
+		FVector& ImpactLocationWS)
+	{
+		FVector ImpactNormal;
+		const bool bFound = FindPlantTraceImpact(Context, TraceSettings, bCheckComplex, StartPositionWS, ImpactLocationWS, ImpactNormal);
+		OutPlantPlaneWS = FPlane(ImpactLocationWS, ImpactNormal);
+
+		return bFound;
+	}
+
+	static FVector CalculateCentroid(const TArrayView<FTransform>& Transforms)
+	{
+		check(Transforms.Num() > 0);
+
+		FVector Centroid = FVector::ZeroVector;
+		for (const FTransform& Transform : Transforms)
+		{
+			Centroid += Transform.GetLocation();
+		}
+
+		Centroid /= float(Transforms.Num());
+		return Centroid;
+	}
+
+	static float GetDistanceToPlaneAlongDirection(const FVector& Location, const FPlane& PlantPlane, const FVector& ApproachDir)
+	{
+		const FVector IntersectionLoc = FMath::LinePlaneIntersection(
+			Location,
+			Location - ApproachDir,
+			PlantPlane);
+
+		const FVector IntersectionToLocation = Location - IntersectionLoc;
+		const float DistanceToPlantPlane = IntersectionToLocation | -ApproachDir;
+		return DistanceToPlantPlane;
+	}
+
+	static void FindChainLengthRootBoneIndex(
+		const FCompactPoseBoneIndex& InFootBoneIndex,
+		const int32& NumBonesInLimb,
+		const FBoneContainer& RequiredBones,
+		FCompactPoseBoneIndex& OutHipIndex,
+		float& OutChainLength)
+	{
+		OutChainLength = 0.0f;
+		FCompactPoseBoneIndex BoneIndex = InFootBoneIndex;
+		if (BoneIndex != INDEX_NONE)
+		{
+			FCompactPoseBoneIndex ParentBoneIndex = RequiredBones.GetParentBoneIndex(BoneIndex);
+
+
+			int32 NumIterations = NumBonesInLimb;
+			while ((NumIterations-- > 0) && (ParentBoneIndex != INDEX_NONE))
+			{
+				const FTransform BoneTransformPS = RequiredBones.GetRefPoseTransform(BoneIndex);
+				const float Extension = BoneTransformPS.GetTranslation().Size();
+				OutChainLength += Extension;
+
+				BoneIndex = ParentBoneIndex;
+				ParentBoneIndex = RequiredBones.GetParentBoneIndex(BoneIndex);
+			};
+		}
+
+		OutHipIndex = BoneIndex;
+	}
+};
+
+void FAnimNode_FootPlacement::FindPelvisOffsetRangeForLimb(
+	const UE::Anim::FootPlacement::FEvaluationContext& Context,
+	const UE::Anim::FootPlacement::FLegRuntimeData::FInputPoseData& LegInputPose,
+	const FVector& PlantTargetLocationCS,
+	const FTransform& PelvisTransformCS,
+	const float LimbLength,
+	FFloatRange& OutPelvisOffsetRangeCS) const
+{
+	// TODO: Cache this.
+	const FTransform HipToPelvis =
+		LegInputPose.HipTransformCS.GetRelativeTransform(PelvisData.InputPose.FKTransformCS);
+	const FTransform HipTransformCS = HipToPelvis * PelvisTransformCS;
+	const FVector HipLocationCS = HipTransformCS.GetLocation();
+
+	const FVector DesiredExtensionDelta =
+		LegInputPose.FKTransformCS.GetLocation() - LegInputPose.HipTransformCS.GetLocation();
+
+	OutPelvisOffsetRangeCS = FFloatRange::Inclusive(0.0f, 0.0f);
+
+	const float DesiredExtensionSqrd = DesiredExtensionDelta.SizeSquared();
+	const float DesiredExtension = FMath::Sqrt(DesiredExtensionSqrd);
+	const float MaxExtension = GetMaxLimbExtension(DesiredExtension, LimbLength);
+
+	const FVector HipToPlantCS = PlantTargetLocationCS - HipLocationCS;
+	const float HipToPlantDotApproachDir = HipToPlantCS | Context.ApproachDirCS;
+
+	FVector DesiredPlantTargetLocationCS = PlantTargetLocationCS;
+	FVector MaxPlantTargetLocationCS = PlantTargetLocationCS;
+	// If the foot wants to be place so high up relative to the FK hip, this is unlikely to matter.
+	// TODO: Crouch might be an exception. Revisit when looking into that.
+	if (HipToPlantDotApproachDir > 0.0f)
+	{
+		const float OpposingSideSqrd = HipToPlantCS.SizeSquared() - HipToPlantDotApproachDir * HipToPlantDotApproachDir;
+		const float OpposingSide = FMath::Sqrt(OpposingSideSqrd);
+
+		const float MaxExtensionSqrd = MaxExtension * MaxExtension;
+
+		const FPlane FootPlane = FPlane(PlantTargetLocationCS, Context.ApproachDirCS);
+		const FVector FKFootProjected = FVector::PointPlaneProject(LegInputPose.FKTransformCS.GetLocation(), FootPlane);
+		const FVector HipProjected = FVector::PointPlaneProject(HipLocationCS, FootPlane);
+
+		const float MaxOffset = PelvisSettings.MaxOffsetHorizontal;
+		const FVector IKFootToHip = HipProjected - PlantTargetLocationCS;
+		const float IKFootToHipDist = IKFootToHip.Size();
+		const float FKFootToHipDist = FVector::Dist(HipProjected, FKFootProjected);
+		auto FindPlantLocationAdjustedByOrtogonalLimit = 
+			[	PlantTargetLocationCS, MaxOffset, IKFootToHip, IKFootToHipDist,
+				HipToPlantDotApproachDir, OpposingSide, FKFootToHipDist	]
+			(const float RadiusSqrd)
+		{
+			FVector AdjustedPlantTargetLocationCS = PlantTargetLocationCS;
+
+			// The desired height at the limit of  max horizontal extension
+			const float DesiredHeight = HipToPlantDotApproachDir - MaxOffset;
+			const float DesiredHeightSqrd = DesiredHeight * DesiredHeight;
+
+			// Find the max horizontal offset. 
+			// We don't care about the circle intersection in the opposite direction
+			const float MaxOpposingSide = FMath::Sqrt(-DesiredHeightSqrd + RadiusSqrd);
+
+			// Respected the input pose if it exceeds it
+			const float MaxIKOrthogonalDist = FMath::Max(FKFootToHipDist, MaxOpposingSide);
+
+			if (IKFootToHipDist > MaxIKOrthogonalDist)
+			{
+				// Move the foot towards the projected hip
+				AdjustedPlantTargetLocationCS += (IKFootToHipDist - MaxIKOrthogonalDist) * IKFootToHip.GetSafeNormal();
+			}
+
+			return AdjustedPlantTargetLocationCS;
+		};
+
+		MaxPlantTargetLocationCS = FindPlantLocationAdjustedByOrtogonalLimit(MaxExtensionSqrd);
+		DesiredPlantTargetLocationCS = FindPlantLocationAdjustedByOrtogonalLimit(DesiredExtensionSqrd);
+	}
+
+	// Taken from http://runevision.com/thesis/rune_skovbo_johansen_thesis.pdf
+	// Chapter 7.4.2
+	//	Intersections are found of a vertical line going through the original hip
+	//	position and two spheres with their centers at the new ankle position (PlantTargetLocationCS) 
+	//	Sphere 1 has a radius of the distance between the hip and ankle in the input pose (DesiredExtension)
+	//	Sphere 2 has a radius corresponding to the length of the leg from hip to ankle (MaxExtension).
+	FVector MaxOffsetLocation;
+	FVector DesiredOffsetLocation;
+	FMath::SphereDistToLine(MaxPlantTargetLocationCS, MaxExtension, HipLocationCS - Context.ApproachDirCS * TraceSettings.EndOffset, Context.ApproachDirCS, MaxOffsetLocation);
+	FMath::SphereDistToLine(DesiredPlantTargetLocationCS, DesiredExtension, HipLocationCS - Context.ApproachDirCS * TraceSettings.EndOffset, Context.ApproachDirCS, DesiredOffsetLocation);
+	const float MaxOffset = (MaxOffsetLocation - HipLocationCS) | -Context.ApproachDirCS;
+	const float DesiredOffset = (DesiredOffsetLocation - HipLocationCS) | -Context.ApproachDirCS;
+
+	OutPelvisOffsetRangeCS = FFloatRange::Inclusive(DesiredOffset, MaxOffset);
+}
+
+float FAnimNode_FootPlacement::CalcTargetPlantPlaneDistance(
+	const UE::Anim::FootPlacement::FEvaluationContext& Context,
+	const UE::Anim::FootPlacement::FLegRuntimeData::FInputPoseData& LegInputPose) const
+{
+	const FTransform IKBallBoneCS = LegInputPose.FootToBall * LegInputPose.IKTransformCS;
+
+	const FTransform& IKFootRootCS = PelvisData.InputPose.IKRootTransformCS;
+	const FPlane IKGroundPlaneCS = 
+		FPlane(	PelvisData.InputPose.IKRootTransformCS.GetLocation(), 
+				PelvisData.InputPose.IKRootTransformCS.TransformVectorNoScale(FVector::UpVector));
+
+	// TODO: I'm just getting the distance between bones and the plane, instead of actual foot/ball bases
+	const float FootBaseDistance = 
+		UE::Anim::FootPlacement::GetDistanceToPlaneAlongDirection(LegInputPose.IKTransformCS.GetLocation(), IKGroundPlaneCS, Context.ApproachDirCS);
+	const float BallBaseDistance = 
+		UE::Anim::FootPlacement::GetDistanceToPlaneAlongDirection(IKBallBoneCS.GetLocation(), IKGroundPlaneCS, Context.ApproachDirCS);
+
+	const float PlantPlaneDistance = FMath::Min(FootBaseDistance, BallBaseDistance);
+	return PlantPlaneDistance;
+}
+
+
+void FAnimNode_FootPlacement::AlignPlantToGround(
+	const UE::Anim::FootPlacement::FEvaluationContext& Context,
+	const FPlane& PlantPlaneWS,
+	const UE::Anim::FootPlacement::FLegRuntimeData::FInputPoseData& LegInputPose,
+	FTransform& InOutFootTransformWS,
+	FQuat& OutTwistCorrection) const
+{
+	const FTransform InputPoseFootTransformWS = LegInputPose.IKTransformCS * Context.OwningComponentToWorld;
+
+	// It is assumed the distance from the plane defined by ik foot root to the ik reference, along the trace 
+	// direction, must remain the same.
+	// TODO: This wont work well when the animation doesn't have a single plant plane, i.e. a walking upstairs anim
+	const FTransform IKFootRootWS = PelvisData.InputPose.IKRootTransformCS * Context.OwningComponentToWorld;
+	const FPlane IKFootRootPlaneWS = FPlane(IKFootRootWS.GetLocation(), IKFootRootWS.TransformVectorNoScale(FVector::UpVector));
+	const float IKFootRootToFootRootTargetDistance = 
+		UE::Anim::FootPlacement::GetDistanceToPlaneAlongDirection(InputPoseFootTransformWS.GetLocation(), IKFootRootPlaneWS, Context.ApproachDirWS);
+
+	const FVector CorrectedPlaneIntersectionWS = FMath::LinePlaneIntersection(
+		InOutFootTransformWS.GetLocation(),
+		InOutFootTransformWS.GetLocation() + Context.ApproachDirWS,
+		PlantPlaneWS);
+
+	const FVector CorrectedLocationWS =
+		CorrectedPlaneIntersectionWS - (Context.ApproachDirWS * IKFootRootToFootRootTargetDistance);
+
+	// The relationship between the ik reference and the normal of the plane defined by the ik foot root must also be 
+	// respected
+	const FQuat PlanePlaneDeltaRotation = FQuat::FindBetweenNormals(IKFootRootPlaneWS.GetNormal(), PlantPlaneWS.GetNormal());
+	const FQuat InputPoseAlignedRotationWS = PlanePlaneDeltaRotation * InputPoseFootTransformWS.GetRotation();
+
+	// Find the rotation that will take us from the Aligned Input Pose to the Unaligned IK Foot 
+	const FQuat UnalignedIKFootToUnalignedInputPoseRotationDelta =
+		InputPoseAlignedRotationWS.Inverse() * InOutFootTransformWS.GetRotation();
+	const FVector IKReferenceNormalFootSpace = InputPoseAlignedRotationWS.UnrotateVector(PlantPlaneWS.GetNormal());
+
+	// Calculate and apply the amount of twist around the IK Root plane. 
+	// This is also used to calculate lock rotation limits
+	FQuat OutSwing;
+	UnalignedIKFootToUnalignedInputPoseRotationDelta.ToSwingTwist(IKReferenceNormalFootSpace,
+		OutSwing,
+		OutTwistCorrection);
+	const FQuat AlignedRotationWS = InputPoseAlignedRotationWS * OutTwistCorrection;
+
+	// Find the rotation that will take us from aligned to unaligned foot
+	const FQuat AlignedToUnalignedRotationDelta =
+		AlignedRotationWS.Inverse() * InOutFootTransformWS.GetRotation();
+	// The rotation is a delta so we won't need to re-orient this vector
+	const FVector FootToBallDir = LegInputPose.FootToBall.GetTranslation().GetSafeNormal();
+	FQuat AnkleTwist;
+	AlignedToUnalignedRotationDelta.ToSwingTwist(FootToBallDir,
+		OutSwing,
+		AnkleTwist);
+	// Counter the aligned ankle twist by the user-defined amount
+	const FQuat TwistCorrectedRotationWS = AlignedRotationWS *  FQuat::Slerp(FQuat::Identity, AnkleTwist, PlantSettings.AnkleTwistReduction);
+
+	// TODO: Clipping will occur due to rotation. Figure out how much we need to adjust the foot vertically to prevent clipping.
+
+	InOutFootTransformWS = FTransform(TwistCorrectedRotationWS, CorrectedLocationWS);;
+}
+
+FTransform FAnimNode_FootPlacement::UpdatePlantOffsetInterpolation(
+	const UE::Anim::FootPlacement::FEvaluationContext& Context,
+	UE::Anim::FootPlacement::FLegRuntimeData::FInterpolationData& InOutInterpData,
+	const FTransform& DesiredTransformCS) const
+{
+	const FTransform IkBase = DesiredTransformCS;
+	const FTransform IkBaseWithCurrOffset = IkBase * InOutInterpData.UnalignedFootOffsetCS;
+
+	// TODO: another indication everything should be done from the base
+	const FVector IkBaseTranslation = IkBaseWithCurrOffset.GetTranslation() - IkBase.GetTranslation();
+	const FVector IkBaseLastSpringTranslation = FVector(FVector2D(IkBaseTranslation), 0.0f);
+
+	const FVector IkBaseSpringTranslation = UKismetMathLibrary::VectorSpringInterp(
+		IkBaseLastSpringTranslation, FVector::ZeroVector, InOutInterpData.PlantOffsetTranslationSpringState,
+		InterpolationSettings.UnplantLinearStiffness,
+		InterpolationSettings.UnplantLinearDamping,
+		Context.UpdateDeltaTime, 1.0f, 0.0f);
+
+	// Since the alignment is just a translation offset, there's no need to calculate a different offset.
+	const FQuat IkBaseSpringRotation = UKismetMathLibrary::QuaternionSpringInterp(
+		InOutInterpData.UnalignedFootOffsetCS.GetRotation(), FQuat::Identity, InOutInterpData.PlantOffsetRotationSpringState,
+		InterpolationSettings.UnplantAngularStiffness,
+		InterpolationSettings.UnplantAngularDamping,
+		Context.UpdateDeltaTime, 1.0f, 0.0f);
+
+	//const FTransform IkBaseSpringOffset = FTransform(IkBaseSpringRotation, IkBaseSpringTranslation);
+	const FTransform IkBaseNewTransform = FTransform(IkBaseSpringRotation * IkBase.GetRotation(),
+		IkBase.GetLocation() + IkBaseSpringTranslation);
+	const FTransform IkBoneNewTransform = IkBaseNewTransform;
+
+	const FTransform BoneTransformOffset = IkBase.GetRelativeTransformReverse(IkBoneNewTransform);
+
+	return BoneTransformOffset;
+}
+
+void FAnimNode_FootPlacement::UpdatePlantingPlaneInterpolation(
+	const UE::Anim::FootPlacement::FEvaluationContext& Context,
+	const FTransform& FootTransformWS,
+	const FTransform& LastAlignedFootTransform,
+	const float AlignmentAlpha,
+	FPlane& InOutPlantPlane,
+	UE::Anim::FootPlacement::FLegRuntimeData::FInterpolationData& InOutInterpData) const
+{
+	const FVector TraceDirection = Context.ApproachDirWS;
+	const FPlane LastPlantPlane = InOutPlantPlane;
+
+	const bool bTraceAgainstSimpleAndComplex = TraceSettings.SimpleCollisionInfluence > 0.0f;
+	if (bTraceAgainstSimpleAndComplex &&!FMath::IsNearlyEqual(AlignmentAlpha, 1.0f))
+	{
+		FVector ImpactLocationSimpleWS; 
+		FVector ImpactLocationComplexWS;
+
+		// Trace against both complex and simple geometry when the foot is in flight. Scale by alignment alpha.
+		// If the collision geometry we're testing against has simple and complex versions, simple collision may provide a smoother path with less clipping.
+		UE::Anim::FootPlacement::FindPlantPlane(Context, TraceSettings, FootTransformWS.GetLocation(), false, InOutPlantPlane, ImpactLocationSimpleWS);
+		UE::Anim::FootPlacement::FindPlantPlane(Context, TraceSettings, FootTransformWS.GetLocation(), true, InOutPlantPlane, ImpactLocationComplexWS);
+
+		// TODO: Alignment alpha is not really what we want. Once we have prediction, and know when the foot will be planted, we can build a better curve.
+		// Simple collision might be better for obstacle avoidance too, since it's presumably already a hull around complex collision.
+		FVector ImpactLocationBlendedWS = FMath::Lerp(
+			ImpactLocationSimpleWS, ImpactLocationComplexWS, 
+			AlignmentAlpha * TraceSettings.SimpleCollisionInfluence + (1.0f - TraceSettings.SimpleCollisionInfluence));
+		InOutPlantPlane = FPlane(ImpactLocationBlendedWS, InOutPlantPlane.GetNormal());
+	}
+	else
+	{
+		FVector ImpactLocationWS;
+		// Trace against complex geometry only to plant accurately
+		UE::Anim::FootPlacement::FindPlantPlane(Context, TraceSettings, FootTransformWS.GetLocation(), true, InOutPlantPlane, ImpactLocationWS);
+	}
+	
+
+	const FVector CurrPlaneIntersection = FMath::LinePlaneIntersection(
+		FootTransformWS.GetLocation(),
+		FootTransformWS.GetLocation() + TraceDirection,
+		InOutPlantPlane);
+
+	const FVector LastPlaneIntersection = FMath::LinePlaneIntersection(
+		LastAlignedFootTransform.GetLocation(),
+		LastAlignedFootTransform.GetLocation() + TraceDirection,
+		LastPlantPlane);
+
+	const FVector PrevPlaneIntersection = FMath::LinePlaneIntersection(
+		FootTransformWS.GetLocation(),
+		FootTransformWS.GetLocation() + TraceDirection,
+		LastPlantPlane);
+
+	const float LastPlaneDeltaZ = LastPlaneIntersection.Z - CurrPlaneIntersection.Z;
+	const float PrevPlaneDeltaZ = PrevPlaneIntersection.Z - CurrPlaneIntersection.Z;
+	const float AdjustedPrevZ = FMath::Abs(LastPlaneDeltaZ) < FMath::Abs(PrevPlaneDeltaZ) ?
+		LastPlaneIntersection.Z : PrevPlaneIntersection.Z;
+
+	//TODO: replace by Z? Do some math and interpolate Plane.W!
+	const FVector AdjustedPrevPlaneIntersection = FVector(CurrPlaneIntersection.X,
+		CurrPlaneIntersection.Y,
+		AdjustedPrevZ);
+
+	const FVector PlantPlaneSpringLocation = UKismetMathLibrary::VectorSpringInterp(
+		AdjustedPrevPlaneIntersection, CurrPlaneIntersection, InOutInterpData.GroundTranslationSpringState,
+		InterpolationSettings.FloorLinearStiffness,
+		InterpolationSettings.FloorLinearDamping,
+		Context.UpdateDeltaTime, 1.0f, 0.0f);
+
+	FQuat FloorNormalRotation = FQuat::FindBetweenNormals(
+		LastPlantPlane.GetNormal(), InOutPlantPlane.GetNormal());
+	const FQuat FloorSpringNormalRotation = UKismetMathLibrary::QuaternionSpringInterp(
+		FQuat::Identity, FloorNormalRotation, InOutInterpData.GroundRotationSpringState,
+		InterpolationSettings.FloorAngularStiffness,
+		InterpolationSettings.FloorAngularDamping,
+		Context.UpdateDeltaTime, 1.0f, 0.0f);
+
+	const FVector PlantPlaneSpringNormal = FloorSpringNormalRotation.RotateVector(LastPlantPlane.GetNormal());
+
+	const FPlane PlantingPlane = FPlane(PlantPlaneSpringLocation, PlantPlaneSpringNormal);
+
+	InOutPlantPlane = PlantingPlane;
+}
+
+void FAnimNode_FootPlacement::DeterminePlantType(
+	const UE::Anim::FootPlacement::FEvaluationContext& Context,
+	const FTransform& FKTransformWS,
+	const FTransform& CurrentBoneTransformWS,
+	UE::Anim::FootPlacement::FLegRuntimeData::FPlantData& InOutPlantData,
+	const UE::Anim::FootPlacement::FLegRuntimeData::FInputPoseData& LegInputPose) const
+{
+	using namespace UE::Anim::FootPlacement;
+
+	const bool bWasPlanted = InOutPlantData.PlantType != EPlantType::Unplanted;
+	const bool bWantedToPlant = InOutPlantData.bWantsToPlant;
+
+	InOutPlantData.bWantsToPlant = WantsToPlant(Context, LegInputPose);
+	InOutPlantData.PlantType = EPlantType::Unplanted;
+
+	if (!InOutPlantData.bWantsToPlant || !InOutPlantData.bCanReachTarget)
+	{
+		return;
+	}
+
+	// Test for un-plant
+	if (bWasPlanted)
+	{
+		const FVector PlantTranslationWS =
+			CurrentBoneTransformWS.GetLocation() - FKTransformWS.GetLocation();
+
+		// TODO: Test along approach direction
+		const bool bPlantTranslationExceeded =
+			PlantTranslationWS.SizeSquared2D() > PlantRuntimeSettings.MaxLinearErrorSqrd;
+		const bool bPlantRotationExceeded =
+			FMath::Abs(InOutPlantData.TwistCorrection.W) <
+			PlantRuntimeSettings.CosHalfMaxRotationError;
+
+		if (!bPlantTranslationExceeded && !bPlantRotationExceeded)
+		{
+			// Carry over result from last plant.
+			InOutPlantData.PlantType = InOutPlantData.LastPlantType;
+		}
+	}
+	else if (!bWantedToPlant)
+	{
+		// If FK wasn't planted last frame, and it is on this frame, we're planted 
+		InOutPlantData.PlantType = UE::Anim::FootPlacement::EPlantType::Planted;
+	}
+	else // Test for re-plant
+	{
+		const FVector PlantLocationDelta =
+			CurrentBoneTransformWS.GetLocation() - FKTransformWS.GetLocation();
+
+		// TODO: Test along approach direction
+		const float LocationDeltaSizeSqrd = PlantLocationDelta.SizeSquared2D();
+
+		const bool bLocationWithinBounds =
+			LocationDeltaSizeSqrd <= PlantRuntimeSettings.ReplantMaxLinearErrorSqrd;
+		const bool bTwistWithinBounds =
+			FMath::Abs(InOutPlantData.TwistCorrection.W) >=
+			PlantRuntimeSettings.CosHalfReplantMaxRotationError;
+
+		if (bLocationWithinBounds && bTwistWithinBounds)
+		{
+			InOutPlantData.PlantType = UE::Anim::FootPlacement::EPlantType::Replanted;
+		}
+	}
+}
+
+float FAnimNode_FootPlacement::GetMaxLimbExtension(const float DesiredExtension, const float LimbLength) const
+{
+	if (DesiredExtension > LimbLength)
+	{
+		return DesiredExtension;
+	}
+
+	const float RemainingLength = LimbLength - DesiredExtension;
+	return DesiredExtension + RemainingLength * PlantSettings.MaxExtensionRatio;
+}
+
+float FAnimNode_FootPlacement::GetMinLimbExtension(const float DesiredExtension, const float LimbLength) const
+{
+	return FMath::Min(DesiredExtension, LimbLength * PlantSettings.MinExtensionRatio);
+}
+
+bool FAnimNode_FootPlacement::WantsToPlant(
+	const UE::Anim::FootPlacement::FEvaluationContext& Context,
+	const UE::Anim::FootPlacement::FLegRuntimeData::FInputPoseData& LegInputPose) const
+{
+#if ENABLE_ANIM_DEBUG
+	if (!CVarAnimNodeFootPlacementEnableLock.GetValueOnAnyThread())
+	{
+		return false;
+	}
+#endif
+
+	if (PlantSettings.LockType == EFootPlacementLockType::Unlocked)
+	{
+		return false;
+	}
+
+	const bool bPassesPlantDistanceCheck = LegInputPose.DistanceToPlant < PlantSettings.FKPlantDistance;
+	const bool bPassesSpeedCheck = LegInputPose.Speed < PlantSettings.SpeedThreshold;
+	return bPassesPlantDistanceCheck && bPassesSpeedCheck;
+}
+
+float FAnimNode_FootPlacement::GetAlignmentAlpha(
+	const UE::Anim::FootPlacement::FEvaluationContext& Context,
+	const UE::Anim::FootPlacement::FLegRuntimeData::FInputPoseData& LegInputPose) const
+{
+	return FMath::Clamp(
+		FMath::GetRangePct(FVector2D(PlantSettings.UnalignmentSpeedThreshold,
+			PlantSettings.SpeedThreshold),
+			LegInputPose.Speed),
+		0.0f, 1.0f);
+}
+
+FTransform FAnimNode_FootPlacement::GetFootPivotAroundBallWS(
+	const UE::Anim::FootPlacement::FEvaluationContext& Context,
+	const UE::Anim::FootPlacement::FLegRuntimeData::FInputPoseData& LegInputPose,
+	const FTransform& LastPlantTransformWS) const
+{
+	const FTransform BallTransformWS = LegInputPose.BallTransformCS * Context.OwningComponentToWorld;
+
+	const FTransform PinnedBallTransformWS = FTransform(
+		BallTransformWS.GetRotation(),
+		(LegInputPose.FootToBall * LastPlantTransformWS).GetLocation(),
+		BallTransformWS.GetScale3D());
+
+	return LegInputPose.BallToFoot * PinnedBallTransformWS;
+
+}
+
+UE::Anim::FootPlacement::FPlantResult FAnimNode_FootPlacement::FinalizeFootAlignment(
+	const UE::Anim::FootPlacement::FEvaluationContext& Context,
+	UE::Anim::FootPlacement::FLegRuntimeData& LegData,
+	const FFootPlacemenLegDefinition& LegDef,
+	const FTransform& PelvisTransformCS)
+{
+	// TODO: Cache this value
+	const FTransform& FKPelvisToHipCS =
+		LegData.InputPose.HipTransformCS.GetRelativeTransform(PelvisData.InputPose.FKTransformCS);
+	const FTransform FinalHipTransformCS = FKPelvisToHipCS * PelvisTransformCS;
+	FTransform CorrectedFootTransformCS = LegData.AlignedFootTransformCS;
+
+	// avoid hyper extension - start
+	const FVector InitialHipToFootDir =
+		(LegData.InputPose.FKTransformCS.GetLocation() - LegData.InputPose.HipTransformCS.GetLocation()).GetSafeNormal();
+	const FVector TargetHipToFootDir =
+		(CorrectedFootTransformCS.GetLocation() - FinalHipTransformCS.GetLocation()).GetSafeNormal();
+
+	// Assume the pelvis adjustments let us reach the spot, unless we're too over-extended
+	LegData.Plant.bCanReachTarget = true;
+
+	const FVector FKHipToFoot =
+		LegData.InputPose.FKTransformCS.GetLocation() - LegData.InputPose.HipTransformCS.GetLocation();
+
+	if (!InitialHipToFootDir.IsNearlyZero() && !TargetHipToFootDir.IsNearlyZero())
+	{
+		{
+
+			const float FKExtension = FKHipToFoot.Size();
+			const float MaxExtension = GetMaxLimbExtension(FKExtension, LegData.Bones.LimbLength);
+			const float IKExtension =
+				FVector::Dist(CorrectedFootTransformCS.GetLocation(), FinalHipTransformCS.GetLocation());
+
+			const float HyperExtensionAmount = IKExtension - MaxExtension;
+			float HyperExtensionRemaining = HyperExtensionAmount;
+
+			if (IKExtension > MaxExtension)
+			{
+				const bool bIsPlanted = LegData.Plant.PlantType != UE::Anim::FootPlacement::EPlantType::Unplanted;
+				const bool bWasPlanted = LegData.Plant.LastPlantType != UE::Anim::FootPlacement::EPlantType::Unplanted;
+				const bool bPlantedThisFrame = bIsPlanted && !bWasPlanted;
+
+				if (!bIsPlanted)
+				{
+					// If there's any overextension and we're unplanted, target is unreachable
+					// Don't plant until we're in re-plant range
+					LegData.Plant.bCanReachTarget = false;
+				}
+
+				const FTransform FKHipToLeg =
+					LegData.InputPose.FKTransformCS.GetRelativeTransform(PelvisData.InputPose.FKTransformCS);
+				const FTransform FKLegAtCurrentHipCS =
+					FKHipToLeg * PelvisTransformCS;
+				FVector IKToFK = (CorrectedFootTransformCS.GetLocation() - FKLegAtCurrentHipCS.GetLocation()).GetSafeNormal();
+
+				const bool bRecentlyUnplanted = !bIsPlanted && LegData.Plant.TimeSinceFullyUnaligned == 0.0f;
+				// Try to keep the tip on spot if we're unplanting
+				// TODO: Make this configurable? 
+				if (bRecentlyUnplanted || bIsPlanted)
+				{
+					// Scale this value by our FK transition alpha to not pop
+					const float MaxPullTowardsHip =  FMath::Min(LegData.Bones.FootLength, HyperExtensionRemaining) * LegData.InputPose.AlignmentAlpha;
+					HyperExtensionRemaining -= MaxPullTowardsHip;
+
+					const FVector NotHyperextendedPlantLocation = CorrectedFootTransformCS.GetLocation() - TargetHipToFootDir * MaxPullTowardsHip;
+					// Grab the Tip location before adjustments.
+					const FVector HyperextendedIkTipLocationCS = (LegData.InputPose.FootToBall * CorrectedFootTransformCS).GetLocation();
+
+					const FTransform IKBallTransformCorrectedCS = LegData.InputPose.FootToBall * CorrectedFootTransformCS;
+
+					// Try to keep the tip at the same spot
+					const FVector ToTipInitial = (IKBallTransformCorrectedCS.GetLocation() - CorrectedFootTransformCS.GetLocation()).GetSafeNormal();
+					const FVector ToTipDesired = (IKBallTransformCorrectedCS.GetLocation() - NotHyperextendedPlantLocation).GetSafeNormal();
+					FQuat DeltaSlopeRotation = FQuat::FindBetweenNormals(ToTipInitial, ToTipDesired);
+
+					FRotator DeltaSlopeRotator = DeltaSlopeRotation.Rotator();
+
+					CorrectedFootTransformCS.SetRotation(DeltaSlopeRotation * CorrectedFootTransformCS.GetRotation());
+					CorrectedFootTransformCS.NormalizeRotation();
+
+					// Move the IK bone closer to prevent overextension
+					CorrectedFootTransformCS.SetLocation(NotHyperextendedPlantLocation);
+					const FVector NotHyperextendedIkTipLocationCS = (LegData.InputPose.FootToBall * CorrectedFootTransformCS).GetLocation();
+					const float BallDelta = FVector::Dist(HyperextendedIkTipLocationCS, NotHyperextendedIkTipLocationCS);
+				}
+
+				// Fix any remaining hyper-extension
+				if (HyperExtensionRemaining > 0.0f)
+				{
+					// Move IK bone towards the hip bone.
+					// TODO: Pull towards the FK bone? This pull lifts the foot from the ground and it might be 
+					// preferable to slide. This causes discontinuities when the foot is no longer hyper-extended
+					FVector NotHyperextendedPlantLocation;
+					FMath::SphereDistToLine(FinalHipTransformCS.GetLocation(), MaxExtension, CorrectedFootTransformCS.GetLocation(), TargetHipToFootDir, NotHyperextendedPlantLocation);
+					CorrectedFootTransformCS.SetLocation(NotHyperextendedPlantLocation);
+				}
+			}
+
+#if ENABLE_ANIM_DEBUG
+			DebugData.LegsExtension[LegData.Idx].HyperExtensionAmount = HyperExtensionAmount;
+			DebugData.LegsExtension[LegData.Idx].RollAmount = HyperExtensionAmount - HyperExtensionRemaining;
+			DebugData.LegsExtension[LegData.Idx].PullAmount = FMath::Max(0.0f, HyperExtensionRemaining);
+#endif // (ENABLE_ANIM_DEBUG)
+		}
+	}
+
+	// Next the plant is ajusted to prevent penetration with the planting plane. To do that, first the base of the plant
+	// and the tip must be calculated (note that because the ground plane interpolates, this does not prevent physical penetration 
+	// with the geometry).
+	const FTransform CorrectedBallTransformCS = LegData.InputPose.FootToBall * CorrectedFootTransformCS;
+
+	// TODO: Consolidate with CalcTargetPlantPlaneDistance
+	const float FootDistance = UE::Anim::FootPlacement::GetDistanceToPlaneAlongDirection(
+		CorrectedBallTransformCS.GetLocation(),
+		LegData.Plant.PlantPlaneCS,
+		Context.ApproachDirCS);
+	const float BallDistance = UE::Anim::FootPlacement::GetDistanceToPlaneAlongDirection(
+		CorrectedFootTransformCS.GetLocation(),
+		LegData.Plant.PlantPlaneCS,
+		Context.ApproachDirCS);
+	const float MinDistance = FMath::Min(FootDistance, BallDistance);
+
+	// A min distance < 0.0f means there was penetration
+	if (MinDistance < 0.0f)
+	{
+		CorrectedFootTransformCS.AddToTranslation(MinDistance * Context.ApproachDirCS);
+	}
+
+	// Fix any remaining hyper-compression. Clip into the ground plane if necessary.
+	// Doing this after pushing the feet out of the ground plane ensures we won't end up in awkward poses. 
+	{
+		FVector NotHyperextendedPlantLocation;
+		const float MinExtension = GetMinLimbExtension(FMath::Abs(FKHipToFoot | Context.ApproachDirCS), LegData.Bones.LimbLength);
+
+		// Offset our hip plane by min extension
+		FPlane HipPlane = FPlane(FinalHipTransformCS.GetLocation() + Context.ApproachDirCS * MinExtension, Context.ApproachDirCS);
+		const float DistanceToHipPlane = HipPlane.PlaneDot(CorrectedFootTransformCS.GetLocation());
+
+		if (DistanceToHipPlane < 0.0f)
+		{
+			// Move foot to hip plane if we're past it.
+			NotHyperextendedPlantLocation = CorrectedFootTransformCS.GetLocation() - Context.ApproachDirCS * DistanceToHipPlane;
+			CorrectedFootTransformCS.SetLocation(NotHyperextendedPlantLocation);
+		}
+	}
+
+	// TODO: Do adjustments to FK tip and FK Chain root
+	FTransform FinalFkTipTransformCS = FTransform::Identity;
+	FTransform FinalFkHipTransformCS = FTransform::Identity;
+
+	UE::Anim::FootPlacement::FPlantResult Result =
+	{
+		{ LegData.Bones.IKIndex, CorrectedFootTransformCS },
+		//{ LegData.Bones.BallIndex, FinalFkTipTransformCS },
+		//{ LegData.Bones.HipIndex, FinalFkHipTransformCS }
+	};
+
+	return Result;
+}
+
+#if ENABLE_ANIM_DEBUG
+void FAnimNode_FootPlacement::DrawDebug(
+	const UE::Anim::FootPlacement::FEvaluationContext& Context,
+	const UE::Anim::FootPlacement::FLegRuntimeData& LegData,
+	const UE::Anim::FootPlacement::FPlantResult& PlantResult) const
+{
+	using namespace UE::Anim::FootPlacement;
+
+	const FColor FKColor = FColor::Blue;
+	const FColor PlantedColor = FColor::Red;
+	const FColor UnplantedColor = FColor::Green;
+	const FColor ReplantedColor = FColor::Orange;
+
+	FColor CurrentPlantColor;
+	switch (LegData.Plant.PlantType)
+	{
+	case UE::Anim::FootPlacement::EPlantType::Planted: CurrentPlantColor = PlantedColor; break;
+	case UE::Anim::FootPlacement::EPlantType::Unplanted: CurrentPlantColor = UnplantedColor; break;
+	case UE::Anim::FootPlacement::EPlantType::Replanted: CurrentPlantColor = ReplantedColor; break;
+	default: check(false); break; //not implemented
+	}
+
+	const FTransform FKBoneTransformWS =
+		LegData.InputPose.FootToGround *
+		LegData.InputPose.FKTransformCS *
+		Context.OwningComponentToWorld;
+
+	const FTransform IKBoneTransformWS =
+		LegData.InputPose.FootToGround *
+		LegData.AlignedFootTransformWS;
+
+
+	const FVector FKBoneLocationProjectedWS = FMath::LinePlaneIntersection(
+		FKBoneTransformWS.GetLocation(),
+		FKBoneTransformWS.GetLocation() + Context.ApproachDirWS,
+		LegData.Plant.PlantPlaneWS);
+
+	Context.CSPContext.AnimInstanceProxy->AnimDrawDebugPoint(
+		FKBoneTransformWS.GetLocation(), 10.0f, FKColor, false, -1.0f, SDPG_Foreground);
+	Context.CSPContext.AnimInstanceProxy->AnimDrawDebugPoint(
+		FKBoneLocationProjectedWS, 15.0f, FKColor, false, -1.0f, SDPG_Foreground);
+	Context.CSPContext.AnimInstanceProxy->AnimDrawDebugLine(
+		FKBoneTransformWS.GetLocation(), FKBoneLocationProjectedWS,
+		FKColor,
+		false, -1.0f, 1.0f, SDPG_Foreground);
+	Context.CSPContext.AnimInstanceProxy->AnimDrawDebugPoint(
+		IKBoneTransformWS.GetLocation(), 10.0f, CurrentPlantColor, false, -1.0f, SDPG_Foreground);
+
+	const FVector IKBoneLocationProjectedWS = FMath::LinePlaneIntersection(
+		IKBoneTransformWS.GetLocation(),
+		IKBoneTransformWS.GetLocation() + Context.ApproachDirWS,
+		LegData.Plant.PlantPlaneWS);
+
+	Context.CSPContext.AnimInstanceProxy->AnimDrawDebugPoint(
+		IKBoneLocationProjectedWS, 15.0f, CurrentPlantColor, false, -1.0f, SDPG_Foreground);
+	Context.CSPContext.AnimInstanceProxy->AnimDrawDebugLine(
+		IKBoneTransformWS.GetLocation(), IKBoneLocationProjectedWS,
+		CurrentPlantColor,
+		false, -1.0f, 1.0f, SDPG_Foreground);
+
+	const float UnplantRadius = PlantSettings.MaxLinearError;
+	const FVector PlantCenter = FMath::LinePlaneIntersection(
+		IKBoneTransformWS.GetLocation(),
+		IKBoneTransformWS.GetLocation() + Context.ApproachDirWS,
+		LegData.Plant.PlantPlaneWS);
+	Context.CSPContext.AnimInstanceProxy->AnimDrawDebugCircle(
+		PlantCenter, UnplantRadius, 24, PlantedColor,
+		LegData.Plant.PlantPlaneWS.GetNormal(), false, -1.0f, SDPG_Foreground, 0.5f);
+
+	const float ReplantRadius =
+		PlantSettings.MaxLinearError *
+		PlantSettings.ReplantMaxLinearErrorRatio;
+	Context.CSPContext.AnimInstanceProxy->AnimDrawDebugCircle(
+		PlantCenter, ReplantRadius, 24, ReplantedColor,
+		LegData.Plant.PlantPlaneWS.GetNormal(), false, -1.0f, SDPG_Foreground, 0.5f);
+
+	
+	FString InputPoseMessage = FString::Format(
+		TEXT("{0}\n\t - InputPose [ AlignmentAlpha = {1}, Speed = {2}, DistanceToPlant = {3}]"), 
+		{	*LegDefinitions[LegData.Idx].FKFootBone.BoneName.ToString(),
+			LegData.InputPose.AlignmentAlpha,
+			LegData.InputPose.Speed,
+			LegData.InputPose.DistanceToPlant });
+	Context.CSPContext.AnimInstanceProxy->AnimDrawDebugOnScreenMessage(InputPoseMessage, FColor::White);
+	
+	FString ExtensionMessage = FString::Format(
+		TEXT("\t - HyperExtension[ Amount = {1}, Roll = {2}, Pull {3}]"),
+		{	 DebugData.LegsExtension[LegData.Idx].HyperExtensionAmount,
+			DebugData.LegsExtension[LegData.Idx].RollAmount,
+			DebugData.LegsExtension[LegData.Idx].PullAmount });
+	Context.CSPContext.AnimInstanceProxy->AnimDrawDebugOnScreenMessage(ExtensionMessage,
+		(DebugData.LegsExtension[LegData.Idx].HyperExtensionAmount <= 0.0f) ? FColor::Green : FColor::Red);
+
+	TRACE_ANIM_NODE_VALUE(Context.CSPContext, TEXT("HyperExtension - Amount"), DebugData.LegsExtension[LegData.Idx].HyperExtensionAmount);
+	TRACE_ANIM_NODE_VALUE(Context.CSPContext, TEXT("HyperExtension - Roll"), DebugData.LegsExtension[LegData.Idx].RollAmount);
+	TRACE_ANIM_NODE_VALUE(Context.CSPContext, TEXT("HyperExtension - Pull"), DebugData.LegsExtension[LegData.Idx].PullAmount);
+	TRACE_ANIM_NODE_VALUE(Context.CSPContext, TEXT("InputPose - AlignmentAlpha"), LegData.InputPose.AlignmentAlpha);
+}
+#endif
+
+FAnimNode_FootPlacement::FAnimNode_FootPlacement()
+{
+}
+
+// TODO: implement 
+void FAnimNode_FootPlacement::GatherDebugData(FNodeDebugData& NodeDebugData)
+{
+
+}
+
+void FAnimNode_FootPlacement::Initialize_AnyThread(const FAnimationInitializeContext& Context)
+{
+	FAnimNode_SkeletalControlBase::Initialize_AnyThread(Context);
+
+	PelvisData.Interpolation = UE::Anim::FootPlacement::FPelvisRuntimeData::FInterpolationData();
+
+	LegsData.Reset();
+	LegsData.SetNumUninitialized(LegDefinitions.Num());
+
+	for (int32 LegIndex = 0; LegIndex < LegsData.Num(); ++LegIndex)
+	{
+		UE::Anim::FootPlacement::FLegRuntimeData& LegData = LegsData[LegIndex];
+		LegData.Idx = LegIndex;
+		LegData.Interpolation = UE::Anim::FootPlacement::FLegRuntimeData::FInterpolationData();
+	}
+
+#if ENABLE_ANIM_DEBUG
+	DebugData.Init(LegDefinitions.Num());
+#endif
+
+	bIsFirstUpdate = true;
+}
+
+void FAnimNode_FootPlacement::UpdateInternal(const FAnimationUpdateContext& Context)
+{
+	FAnimNode_SkeletalControlBase::UpdateInternal(Context);
+	CachedDeltaTime += Context.GetDeltaTime();
+}
+
+void FAnimNode_FootPlacement::EvaluateSkeletalControl_AnyThread(FComponentSpacePoseContext& Output,
+	TArray<FBoneTransform>& OutBoneTransforms)
+{
+	SCOPE_CYCLE_COUNTER(STAT_FootPlacement_Eval);
+
+	check(OutBoneTransforms.Num() == 0);
+
+	UE::Anim::FootPlacement::FDebugData LastDebugData = DebugData;
+	// TODO: Support a different approach direction
+	const FVector ApproachDirCS = -FVector::UpVector;
+	UE::Anim::FootPlacement::FEvaluationContext FootPlacementContext(Output, ApproachDirCS, CachedDeltaTime);
+
+	// Gather data from pose and property inputs, and do minimal processing for commonly used values
+	GatherPelvisDataFromInputs(FootPlacementContext);
+
+	for (int32 FootIndex = 0; FootIndex < LegsData.Num(); ++FootIndex)
+	{
+		UE::Anim::FootPlacement::FLegRuntimeData& LegData = LegsData[FootIndex];
+		const FFootPlacemenLegDefinition& LegDef = LegDefinitions[FootIndex];
+		GatherLegDataFromInputs(FootPlacementContext, LegData, LegDef);
+
+		// TODO: All of these can be calculated on initialize, but in case there's value in changing these dynamically,
+		// will keep this for now. If needed, change to lazy update.
+		PlantRuntimeSettings.MaxExtensionRatioSqrd = PlantSettings.MaxExtensionRatio * PlantSettings.MaxExtensionRatio;
+		PlantRuntimeSettings.MinExtensionRatioSqrd = PlantSettings.MinExtensionRatio * PlantSettings.MinExtensionRatio;
+		PlantRuntimeSettings.MaxLinearErrorSqrd = PlantSettings.MaxLinearError * PlantSettings.MaxLinearError;
+		PlantRuntimeSettings.ReplantMaxLinearErrorSqrd =
+			PlantRuntimeSettings.MaxLinearErrorSqrd *
+			PlantSettings.ReplantMaxLinearErrorRatio * PlantSettings.ReplantMaxLinearErrorRatio;
+		PlantRuntimeSettings.CosHalfMaxRotationError = FMath::Cos(FMath::DegreesToRadians(PlantSettings.MaxRotationError / 2.0f));
+		PlantRuntimeSettings.CosHalfReplantMaxRotationError
+			= FMath::Cos(FMath::DegreesToRadians(
+				(PlantSettings.MaxRotationError * PlantSettings.ReplantMaxRotationErrorRatio) / 2.0f));
+	}
+
+	ProcessCharacterState(FootPlacementContext);
+
+	for (UE::Anim::FootPlacement::FLegRuntimeData& LegData : LegsData)
+	{
+		ProcessFootAlignment(FootPlacementContext, LegData);
+	}
+
+	// Based on the ground alignment, search for the best Pelvis transform
+	FTransform PelvisTransformCS = SolvePelvis(FootPlacementContext);
+	PelvisTransformCS = UpdatePelvisInterpolation(FootPlacementContext, PelvisTransformCS);
+	OutBoneTransforms.Add(FBoneTransform(PelvisData.Bones.FkBoneIndex, PelvisTransformCS));
+
+
+	for (int32 FootIndex = 0; FootIndex < LegsData.Num(); ++FootIndex)
+	{
+		UE::Anim::FootPlacement::FLegRuntimeData& LegData = LegsData[FootIndex];
+		const FFootPlacemenLegDefinition& LegDef = LegDefinitions[FootIndex];
+
+		const UE::Anim::FootPlacement::FPlantResult PlantResult =
+			FinalizeFootAlignment(FootPlacementContext, LegData, LegDef, PelvisTransformCS);
+		OutBoneTransforms.Add(PlantResult.IkPlantTranformCS);
+		//OutBoneTransforms.Add(PlantResult.FkTipTransformCS);
+		//OutBoneTransforms.Add(PlantResult.FkHipTransformCS);
+
+#if ENABLE_ANIM_DEBUG
+		if (CVarAnimNodeFootPlacementDebug.GetValueOnAnyThread())
+		{
+			DrawDebug(FootPlacementContext, LegData, PlantResult);
+
+			// Grab positions to debug draw history
+			DebugData.OutputFootLocationsWS[FootIndex] =
+				FootPlacementContext.OwningComponentToWorld.TransformPosition(PlantResult.IkPlantTranformCS.Transform.GetLocation());
+			DebugData.InputFootLocationsWS[FootIndex] = 
+				FootPlacementContext.OwningComponentToWorld.TransformPosition(LegData.InputPose.IKTransformCS.GetLocation());;
+		}
+#endif
+	}
+
+	OutBoneTransforms.Sort(FCompareBoneTransformIndex());
+
+	CachedDeltaTime = 0.0f;
+
+#if ENABLE_ANIM_DEBUG
+	{
+		FAnimInstanceProxy* AnimInstanceProxy = Output.AnimInstanceProxy;
+		const FTransform ComponentTransform =
+			AnimInstanceProxy->GetSkelMeshComponent()->GetComponentTransform();
+
+		const FVector InputPelvisLocationWS = ComponentTransform.TransformPosition(PelvisData.InputPose.FKTransformCS.GetLocation());
+		const FVector OutputPelvisLocationWS = ComponentTransform.TransformPosition(PelvisTransformCS.GetLocation());
+
+		DebugData.InputPelvisLocationWS = InputPelvisLocationWS;
+		DebugData.OutputPelvisLocationWS = OutputPelvisLocationWS;
+
+		if (CVarAnimNodeFootPlacementDebug.GetValueOnAnyThread())
+		{
+			const int32 DrawIndex = CVarAnimNodeFootPlacementDebugDrawHistory.GetValueOnAnyThread();
+			if ((DrawIndex != 0) && !bIsFirstUpdate)
+			{
+				if (DrawIndex == -1)
+				{
+					AnimInstanceProxy->AnimDrawDebugLine(LastDebugData.OutputPelvisLocationWS, DebugData.OutputPelvisLocationWS, FColor::Magenta, true, -1.0f, 0.5f);
+					AnimInstanceProxy->AnimDrawDebugLine(LastDebugData.InputPelvisLocationWS, DebugData.InputPelvisLocationWS, FColor::Blue, true, -1.0f, 0.5f);
+				}
+				if (DrawIndex > 0 && DebugData.OutputFootLocationsWS.IsValidIndex(DrawIndex - 1))
+				{
+					const int32 FootIndex = DrawIndex - 1;
+					AnimInstanceProxy->AnimDrawDebugLine(DebugData.OutputFootLocationsWS[FootIndex], LastDebugData.OutputFootLocationsWS[FootIndex], FColor::Magenta, true, -1.0f, 0.5f);
+					AnimInstanceProxy->AnimDrawDebugLine(DebugData.InputFootLocationsWS[FootIndex], LastDebugData.InputFootLocationsWS[FootIndex], FColor::Blue, true, -1.0f, 0.5f);
+				}
+			}
+
+			const FTransform PelvisTransformWS = PelvisTransformCS * ComponentTransform;
+			const FTransform BasePelvisTransformWS = PelvisData.InputPose.FKTransformCS * ComponentTransform;
+
+			AnimInstanceProxy->AnimDrawDebugPoint(
+				PelvisTransformWS.GetLocation(), 20.0f, FColor::Green, false, -1.0f, SDPG_Foreground);
+
+			AnimInstanceProxy->AnimDrawDebugPoint(
+				BasePelvisTransformWS.GetLocation(), 20.0f, FColor::Blue, false, -1.0f, SDPG_Foreground);
+		}
+	}
+#endif
+
+	LastComponentLocation = FootPlacementContext.OwningComponentToWorld.GetLocation();
+
+	bIsFirstUpdate = false;
+}
+
+bool FAnimNode_FootPlacement::IsValidToEvaluate(const USkeleton* Skeleton, const FBoneContainer& RequiredBones)
+{
+#if ENABLE_ANIM_DEBUG
+	if (!CVarAnimNodeFootPlacementEnable.GetValueOnAnyThread())
+	{
+		return false;
+	}
+#endif
+
+	for (const UE::Anim::FootPlacement::FLegRuntimeData& LegData : LegsData)
+	{
+		if (!LegData.Bones.HipIndex.IsValid() ||
+			!LegData.Bones.FKIndex.IsValid() ||
+			!LegData.Bones.IKIndex.IsValid() ||
+			!LegData.Bones.BallIndex.IsValid())
+		{
+			return false;
+		}
+	}
+
+	if (!PelvisData.Bones.IkBoneIndex.IsValid() ||
+		!PelvisData.Bones.FkBoneIndex.IsValid())
+	{
+		return false;
+	}
+
+	return true;
+}
+
+void FAnimNode_FootPlacement::InitializeBoneReferences(const FBoneContainer& RequiredBones)
+{
+	for (int32 FootIndex = 0; FootIndex < LegsData.Num(); ++FootIndex)
+	{
+		FFootPlacemenLegDefinition& LegDef = LegDefinitions[FootIndex];
+		UE::Anim::FootPlacement::FLegRuntimeData& LegData = LegsData[FootIndex];
+		LegDef.IKFootBone.Initialize(RequiredBones);
+		LegDef.FKFootBone.Initialize(RequiredBones);
+		LegDef.BallBone.Initialize(RequiredBones);
+
+		LegData.Bones.IKIndex = LegDef.IKFootBone.GetCompactPoseIndex(RequiredBones);
+		LegData.Bones.FKIndex = LegDef.FKFootBone.GetCompactPoseIndex(RequiredBones);
+		LegData.Bones.BallIndex = LegDef.BallBone.GetCompactPoseIndex(RequiredBones);
+		UE::Anim::FootPlacement::FindChainLengthRootBoneIndex(
+			LegData.Bones.FKIndex, FMath::Max(LegDef.NumBonesInLimb, 1), RequiredBones,
+			LegData.Bones.HipIndex, LegData.Bones.LimbLength);
+
+		const FTransform BallTransformLS = RequiredBones.GetRefPoseTransform(LegData.Bones.BallIndex);
+		LegData.Bones.FootLength = BallTransformLS.GetLocation().Size();
+
+		// TODO: This wont work for animations authored for different slopes or stairs. Figure this out later
+		const FVector RefPoseGroundNormalCS = FVector::UpVector;
+		const FTransform BallRefTransformCS = FAnimationRuntime::GetComponentSpaceRefPose(
+			LegData.Bones.BallIndex,
+			RequiredBones);
+		const FVector BallAlignmentDeltaCS = -BallRefTransformCS.GetLocation();
+		const FVector BallAlignmenfOffsetCS = (BallAlignmentDeltaCS | RefPoseGroundNormalCS) * RefPoseGroundNormalCS;
+		LegData.InputPose.BallToGround = FTransform(BallRefTransformCS.GetRotation().UnrotateVector(BallAlignmenfOffsetCS));
+
+		const FTransform FKFootTransformCS = FAnimationRuntime::GetComponentSpaceRefPose(
+			LegData.Bones.FKIndex,
+			RequiredBones);
+		const FVector FootAlignmentDeltaCS = -FKFootTransformCS.GetLocation();
+		const FVector FootAlignmentOffsetCS = (FootAlignmentDeltaCS | RefPoseGroundNormalCS) * RefPoseGroundNormalCS;
+		LegData.InputPose.FootToGround = FTransform(FKFootTransformCS.GetRotation().UnrotateVector(FootAlignmentOffsetCS));
+
+		const USkeleton* Skeleton = RequiredBones.GetSkeletonAsset();
+		check(Skeleton);
+		SmartName::UID_Type NameUID = Skeleton->GetUIDByName(USkeleton::AnimCurveMappingName, LegDef.SpeedCurveName);
+		if (NameUID != SmartName::MaxUID)
+		{
+			// Grab UIDs of filtered curves to avoid lookup later
+			LegData.SpeedCurveUID = NameUID;
+		}
+	}
+
+	PelvisBone.Initialize(RequiredBones);
+	IKFootRootBone.Initialize(RequiredBones);
+
+	PelvisData.Bones.FkBoneIndex = PelvisBone.GetCompactPoseIndex(RequiredBones);
+	PelvisData.Bones.IkBoneIndex = IKFootRootBone.GetCompactPoseIndex(RequiredBones);
+}
+
+void FAnimNode_FootPlacement::GatherPelvisDataFromInputs(const UE::Anim::FootPlacement::FEvaluationContext& Context)
+{
+	PelvisData.InputPose.FKTransformCS =
+		Context.CSPContext.Pose.GetComponentSpaceTransform(PelvisData.Bones.FkBoneIndex);
+	PelvisData.InputPose.IKRootTransformCS =
+		Context.CSPContext.Pose.GetComponentSpaceTransform(PelvisData.Bones.IkBoneIndex);
+
+	// TODO: All of these can be calculated on initialize, but in case there's value in changing these dynamically,
+	// will keep this for now. If needed, change to lazy update.
+	PelvisData.MaxOffsetSqrd = PelvisSettings.MaxOffset * PelvisSettings.MaxOffset;
+	PelvisData.MaxOffsetHorizontalSqrd = PelvisSettings.MaxOffsetHorizontal * PelvisSettings.MaxOffsetHorizontal;
+}
+
+void FAnimNode_FootPlacement::GatherLegDataFromInputs(
+	const UE::Anim::FootPlacement::FEvaluationContext& Context,
+	UE::Anim::FootPlacement::FLegRuntimeData& LegData,
+	const FFootPlacemenLegDefinition& LegDef)
+{
+	LegData.InputPose.FKTransformCS =
+		Context.CSPContext.Pose.GetComponentSpaceTransform(LegData.Bones.FKIndex);
+	LegData.InputPose.IKTransformCS =
+		Context.CSPContext.Pose.GetComponentSpaceTransform(LegData.Bones.IKIndex);
+	LegData.InputPose.BallTransformCS =
+		Context.CSPContext.Pose.GetComponentSpaceTransform(LegData.Bones.BallIndex);
+	LegData.InputPose.HipTransformCS =
+		Context.CSPContext.Pose.GetComponentSpaceTransform(LegData.Bones.HipIndex);
+
+	LegData.InputPose.BallToFoot =
+		LegData.InputPose.FKTransformCS.GetRelativeTransform(LegData.InputPose.BallTransformCS);
+	LegData.InputPose.FootToBall =
+		LegData.InputPose.BallTransformCS.GetRelativeTransform(LegData.InputPose.FKTransformCS);
+
+	if (bIsFirstUpdate)
+	{
+		LegData.AlignedFootTransformWS =
+			LegData.InputPose.FKTransformCS * Context.OwningComponentToWorld;
+		LegData.UnalignedFootTransformWS = LegData.AlignedFootTransformWS;
+
+		const FVector IKFootRootLocationWS = 
+			Context.OwningComponentToWorld.InverseTransformPosition(PelvisData.InputPose.IKRootTransformCS.GetLocation());
+
+		LegData.Plant.PlantPlaneWS = FPlane(IKFootRootLocationWS, -Context.ApproachDirWS);
+		LegData.Plant.PlantPlaneCS = FPlane(PelvisData.InputPose.IKRootTransformCS.GetLocation(), -Context.ApproachDirCS);
+
+		LegData.Plant.PlantType = UE::Anim::FootPlacement::EPlantType::Unplanted;
+		LegData.Plant.LastPlantType = UE::Anim::FootPlacement::EPlantType::Unplanted;
+	}
+
+	bool bValidSpeedCurve;
+	// If the curve is not found in the stream, assume we're unplanted.
+	const float DefaultSpeedCurveValue = PlantSettings.SpeedThreshold;
+	LegData.InputPose.Speed =
+		Context.CSPContext.Curve.Get(LegData.SpeedCurveUID, bValidSpeedCurve, DefaultSpeedCurveValue);
+	LegData.InputPose.DistanceToPlant = CalcTargetPlantPlaneDistance(Context, LegData.InputPose);
+	const float FKAlignmentAlpha = GetAlignmentAlpha(Context, LegData.InputPose);
+	LegData.InputPose.AlignmentAlpha = FKAlignmentAlpha;
+}
+
+void FAnimNode_FootPlacement::ProcessCharacterState(const UE::Anim::FootPlacement::FEvaluationContext& Context)
+{
+	const FVector LastComponentLocationWS = bIsFirstUpdate
+		? Context.OwningComponentToWorld.GetLocation()
+		: CharacterData.ComponentLocationWS;
+
+	CharacterData.ComponentLocationWS = Context.OwningComponentToWorld.GetLocation();
+	CharacterData.NumFKPlanted = 0;
+	for (const UE::Anim::FootPlacement::FLegRuntimeData& LegData : LegsData)
+	{
+		if (FMath::IsNearlyEqual(LegData.InputPose.AlignmentAlpha, 1.0f))
+		{
+			++CharacterData.NumFKPlanted;
+		}
+	}
+
+	const bool bWasOnGround = CharacterData.bIsOnGround;
+	CharacterData.bIsOnGround = !Context.MovementComponent ?
+		true :
+		((Context.MovementComponent->MovementMode == MOVE_Walking) ||
+			(Context.MovementComponent->MovementMode == MOVE_NavWalking)) &&
+		Context.MovementComponent->CurrentFloor.bBlockingHit;
+
+	if (CharacterData.bIsOnGround && bWasOnGround && PelvisSettings.bCompensateForSuddenCapsuleMoves)
+	{
+		// Compensate for sudden capsule moves
+		const FVector CapsuleFloorNormalWS = Context.GetMovementComponentFloorNormal();
+		const FVector OwningComponentAdjustedLastLocationWS =
+			(FMath::Abs(Context.ApproachDirWS | CapsuleFloorNormalWS) > DELTA) ?
+			FMath::LinePlaneIntersection(
+				CharacterData.ComponentLocationWS,
+				CharacterData.ComponentLocationWS + Context.ApproachDirWS,
+				LastComponentLocationWS, CapsuleFloorNormalWS) :
+			CharacterData.ComponentLocationWS;
+
+		const FVector CapsuleMoveOffsetWS =
+			CharacterData.ComponentLocationWS - OwningComponentAdjustedLastLocationWS;
+		if (!CapsuleMoveOffsetWS.IsNearlyZero(KINDA_SMALL_NUMBER))
+		{
+			const FVector CapsuleMoveOffsetCS =
+				Context.OwningComponentToWorld.InverseTransformVectorNoScale(CapsuleMoveOffsetWS);
+			// Offseting our interpolator lets it smoothly solve sudden capsule deltas, instead of following it and pop
+			PelvisData.Interpolation.PelvisTranslationOffset -= CapsuleMoveOffsetCS;
+		}
+	}
+}
+
+void FAnimNode_FootPlacement::ProcessFootAlignment(
+	const UE::Anim::FootPlacement::FEvaluationContext& Context,
+	UE::Anim::FootPlacement::FLegRuntimeData& LegData)
+{
+	UE::Anim::FootPlacement::FLegRuntimeData::FInputPoseData& InputPose = LegData.InputPose;
+	UE::Anim::FootPlacement::FLegRuntimeData::FInterpolationData& Interpolation = LegData.Interpolation;
+	UE::Anim::FootPlacement::FLegRuntimeData::FBoneData& Bones = LegData.Bones;
+	UE::Anim::FootPlacement::FLegRuntimeData::FPlantData& Plant = LegData.Plant;
+
+	const FTransform FKFootTransformWS = InputPose.FKTransformCS * Context.OwningComponentToWorld;
+	const FTransform IKFootTransformWS = InputPose.IKTransformCS * Context.OwningComponentToWorld;
+	const FTransform LastAlignedFootTransformWS = LegData.AlignedFootTransformWS;
+	const FTransform LastUnalignedFootTransformWS = LegData.UnalignedFootTransformWS;
+
+	Plant.LastPlantType = Plant.PlantType;
+	DeterminePlantType(
+		Context,
+		FKFootTransformWS,
+		LastAlignedFootTransformWS,
+		Plant,
+		InputPose);
+
+	const bool bIsPlanted = Plant.PlantType != UE::Anim::FootPlacement::EPlantType::Unplanted;
+	const bool bWasPlanted = Plant.LastPlantType != UE::Anim::FootPlacement::EPlantType::Unplanted;
+
+	if (bIsPlanted)
+	{
+		FTransform CurrentPlantedTransformWS;
+
+		switch (PlantSettings.LockType)
+		{
+		case EFootPlacementLockType::Unlocked:
+			break;
+		case EFootPlacementLockType::PivotAroundBall:
+		{
+			// Figure out the correct foot transform that keeps the ball in place
+			CurrentPlantedTransformWS = GetFootPivotAroundBallWS(Context, InputPose, LastUnalignedFootTransformWS);;
+		}
+			break;
+		case EFootPlacementLockType::PivotAroundAnkle:
+		{
+			// Use the location only
+			CurrentPlantedTransformWS = IKFootTransformWS;
+			CurrentPlantedTransformWS.SetLocation(LastUnalignedFootTransformWS.GetLocation());
+		}
+			break;
+		case EFootPlacementLockType::LockRotation:
+		{
+			// We use the unaligned foot instead of the aligned one
+			// Because we will adjust roll and twist dynamically
+			CurrentPlantedTransformWS = LastUnalignedFootTransformWS;
+		}
+			break;
+		default: check(false); break; //not implemented
+		}
+
+		FTransform PlantedFootTransformCS =
+			CurrentPlantedTransformWS * Context.OwningComponentToWorld.Inverse();
+
+		// The locked transform is aligned to the ground. Conserve the input pose's ground alignment
+		const FVector AlignedBoneLocationCS = PlantedFootTransformCS.GetLocation();
+		const FPlane InputPosePlantPlane = FPlane(InputPose.IKTransformCS.GetLocation(), Context.ApproachDirWS);
+		const FVector UnalignedBoneLocationCS = FVector::PointPlaneProject(AlignedBoneLocationCS, InputPosePlantPlane);
+
+		PlantedFootTransformCS.SetLocation(UnalignedBoneLocationCS);
+
+		// Get the offset relative to the initial foot transform
+		// Reset interpolation 
+		Interpolation.UnalignedFootOffsetCS =
+			InputPose.IKTransformCS.GetRelativeTransformReverse(PlantedFootTransformCS);
+		Interpolation.PlantOffsetTranslationSpringState.Reset();
+		Interpolation.PlantOffsetRotationSpringState.Reset();
+
+		// If we planted, we're fully unaligned
+		Plant.TimeSinceFullyUnaligned = 0.0f;
+	}
+	else
+	{
+		// No plant, so we interpolate the offset out
+		Interpolation.UnalignedFootOffsetCS =
+			UpdatePlantOffsetInterpolation(Context, Interpolation, InputPose.IKTransformCS);
+
+		// If we're unplanted, we know we're fully unaligned the first time we hit zero alignment alpha.
+		if (Plant.TimeSinceFullyUnaligned > 0.0f || FMath::IsNearlyZero(InputPose.AlignmentAlpha))
+		{
+			Plant.TimeSinceFullyUnaligned += Context.UpdateDeltaTime;
+		}
+	}
+
+	const FTransform IKUnalignedTransformCS = InputPose.IKTransformCS * Interpolation.UnalignedFootOffsetCS;
+	LegData.UnalignedFootTransformWS = IKUnalignedTransformCS * Context.OwningComponentToWorld;
+
+	const FTransform ComponentToWorldInv = Context.OwningComponentToWorld.Inverse();
+
+	// find the smooth plant plane
+	UpdatePlantingPlaneInterpolation(Context, LegData.UnalignedFootTransformWS, LastAlignedFootTransformWS, InputPose.AlignmentAlpha, Plant.PlantPlaneWS, Interpolation);
+	Plant.PlantPlaneCS = Plant.PlantPlaneWS.TransformBy(ComponentToWorldInv.ToMatrixWithScale());
+
+	Interpolation.UnalignedFootOffsetCS = InputPose.IKTransformCS.GetRelativeTransformReverse(IKUnalignedTransformCS);
+
+	// This will adjust IkUnalignedTransformWS to make it match the required distance to the plant plane along the 
+	// approach direction, not the plane normal
+	LegData.AlignedFootTransformWS = LegData.UnalignedFootTransformWS;
+	AlignPlantToGround(Context, Plant.PlantPlaneWS, InputPose, LegData.AlignedFootTransformWS, Plant.TwistCorrection);
+
+	LegData.AlignedFootTransformCS =
+		LegData.AlignedFootTransformWS * ComponentToWorldInv;
+
+	// The target transform is a blend based on FkAlignmentAlpha.
+	// Until we have prediction, favor the ground aligned position, 
+	// since this will likely have a more accurate distance from plane
+	FTransform BlendedPlantTransformCS = LegData.AlignedFootTransformCS;
+
+	// When unplanted/unaligned, favor FK orientation and fix penetrations later. 
+	BlendedPlantTransformCS.SetRotation(
+		FQuat::Slerp(
+			InputPose.FKTransformCS.GetRotation(),
+			LegData.AlignedFootTransformCS.GetRotation(),
+			InputPose.AlignmentAlpha));
+
+	LegData.AlignedFootTransformCS = BlendedPlantTransformCS;
+}
+
+FVector FAnimNode_FootPlacement::GetApproachDirWS(const FAnimationBaseContext& Context) const
+{
+	const USkeletalMeshComponent* OwningComponent = Context.AnimInstanceProxy->GetSkelMeshComponent();
+	return -(OwningComponent->GetComponentTransform().GetRotation().GetUpVector());
+}
+
+FTransform FAnimNode_FootPlacement::SolvePelvis(const UE::Anim::FootPlacement::FEvaluationContext& Context)
+{
+	using namespace UE::Anim::FootPlacement;
+
+	// Taken from http://runevision.com/thesis/rune_skovbo_johansen_thesis.pdf
+	// Chapter 7.4.2
+
+	float MaxOffsetMin = BIG_NUMBER;
+	float DesiredOffsetMin = BIG_NUMBER;
+	float DesiredOffsetSum = 0.0f;
+
+	for (const FLegRuntimeData& LegData : LegsData)
+	{
+		FFloatRange PelvisOffsetRangeCS;
+		FindPelvisOffsetRangeForLimb(
+			Context,
+			LegData.InputPose,
+			LegData.AlignedFootTransformCS.GetLocation(),
+			PelvisData.InputPose.FKTransformCS,
+			LegData.Bones.LimbLength,
+			PelvisOffsetRangeCS);
+
+		const float DesiredOffset = PelvisOffsetRangeCS.GetLowerBoundValue();
+		const float MaxOffset = PelvisOffsetRangeCS.GetUpperBoundValue();
+
+		DesiredOffsetSum += DesiredOffset;
+		DesiredOffsetMin = FMath::Min(DesiredOffsetMin, DesiredOffset);
+		MaxOffsetMin = FMath::Min(MaxOffsetMin, MaxOffset);
+	}
+	const float DesiredOffsetAvg = DesiredOffsetSum / LegsData.Num();
+	const float MinToAvg = DesiredOffsetAvg - DesiredOffsetMin;
+	const float MinToMax = MaxOffsetMin - DesiredOffsetMin;
+
+	DesiredOffsetMin -= 0.05f;
+
+	// TODO: This doesn't deal well w/ over-compression. 
+	// In cases like crouching, it favors over-compressing to preserve the pose of the other leg
+	// Consider working in over-compression into the formula.
+	const float Divisor = MinToAvg + MinToMax;
+	const float PelvisOffsetZ = FMath::IsNearlyZero(Divisor) ?
+		DesiredOffsetMin :
+		DesiredOffsetMin + ((MinToAvg * MinToMax) / Divisor);
+
+	const FVector PelvisOffsetDelta = -PelvisOffsetZ * Context.ApproachDirCS;
+	FTransform PelvisTransformCS = PelvisData.InputPose.FKTransformCS;
+	PelvisTransformCS.AddToTranslation(PelvisOffsetDelta);
+
+	return PelvisTransformCS;
+}
+
+FTransform FAnimNode_FootPlacement::UpdatePelvisInterpolation(
+	const UE::Anim::FootPlacement::FEvaluationContext& Context,
+	const FTransform& TargetPelvisTransform)
+{
+	FTransform OutPelvisTransform = TargetPelvisTransform;
+	// Calculate the offset from input pose and interpolate
+	FVector DesiredPelvisOffset =
+		TargetPelvisTransform.GetLocation() - PelvisData.InputPose.FKTransformCS.GetLocation();
+
+	// Clamp by MaxOffset
+	// Clamping the target before interpolation means we may exceed this purely do to interpolation.
+	// If we clamp after, you'll get no smoothing once the limit is reached.
+	const float MaxOffsetSqrd = PelvisData.MaxOffsetSqrd;
+	const float MaxOffset = PelvisSettings.MaxOffset;
+	if (DesiredPelvisOffset.SizeSquared() > MaxOffsetSqrd)
+	{
+		DesiredPelvisOffset = DesiredPelvisOffset.GetClampedToMaxSize(MaxOffset);
+	}
+
+	// Spring interpolation may cause hyperextension/compression so we solve that in FinalizeFootAlignment
+	PelvisData.Interpolation.PelvisTranslationOffset = UKismetMathLibrary::VectorSpringInterp(
+		PelvisData.Interpolation.PelvisTranslationOffset, DesiredPelvisOffset, PelvisData.Interpolation.PelvisTranslationSpringState,
+		PelvisSettings.LinearStiffness,
+		PelvisSettings.LinearDamping,
+		Context.UpdateDeltaTime, 1.0f, 0.0f);
+
+	OutPelvisTransform.SetLocation(
+		PelvisData.InputPose.FKTransformCS.GetLocation() + PelvisData.Interpolation.PelvisTranslationOffset);
+
+	return OutPelvisTransform;
+}
\ No newline at end of file
diff --git a/Engine/Plugins/Animation/AnimationWarping/Source/Runtime/Public/BoneControllers/AnimNode_FootPlacement.h b/Engine/Plugins/Animation/AnimationWarping/Source/Runtime/Public/BoneControllers/AnimNode_FootPlacement.h
new file mode 100644
--- /dev/null
+++ b/Engine/Plugins/Animation/AnimationWarping/Source/Runtime/Public/BoneControllers/AnimNode_FootPlacement.h
@@ -0,0 +1,559 @@
+// Copyright Epic Games, Inc. All Rights Reserved.
+
+#pragma once
+
+#include "BoneControllers/BoneControllerTypes.h"
+#include "BoneControllers/AnimNode_SkeletalControlBase.h"
+#include "Animation/AnimNodeBase.h"
+#include "Kismet/KismetMathLibrary.h"
+#include "Engine/EngineTypes.h"
+
+#include "AnimNode_FootPlacement.generated.h"
+
+namespace UE::Anim::FootPlacement
+{
+	enum class EPlantType
+	{
+		Unplanted,
+		Planted,
+		Replanted
+	};
+
+	struct FLegRuntimeData
+	{
+		int32 Idx = -1;
+
+		struct FBoneData
+		{
+			FCompactPoseBoneIndex FKIndex = FCompactPoseBoneIndex(INDEX_NONE);
+			FCompactPoseBoneIndex BallIndex = FCompactPoseBoneIndex(INDEX_NONE);
+			FCompactPoseBoneIndex IKIndex = FCompactPoseBoneIndex(INDEX_NONE);
+			FCompactPoseBoneIndex HipIndex = FCompactPoseBoneIndex(INDEX_NONE);
+			float LimbLength = 0.0f;
+			float FootLength = 0.0f;
+		} Bones;
+
+		SmartName::UID_Type SpeedCurveUID = SmartName::MaxUID;
+
+		// Helper struct to store values used across the evaluation
+		struct FInputPoseData
+		{
+			FTransform IKTransformCS = FTransform::Identity;
+			FTransform FKTransformCS = FTransform::Identity;
+			FTransform BallTransformCS = FTransform::Identity;
+			FTransform HipTransformCS = FTransform::Identity;
+			FTransform BallToFoot = FTransform::Identity;
+			FTransform FootToBall = FTransform::Identity;
+			FTransform FootToGround = FTransform::Identity;
+			FTransform BallToGround = FTransform::Identity;
+			float Speed = 0.0f;
+			float DistanceToPlant = 0.0f;
+			float AlignmentAlpha = 0.0f;
+		} InputPose;
+
+		/* Ground */
+
+		struct FPlantData
+		{
+			UE::Anim::FootPlacement::EPlantType PlantType = UE::Anim::FootPlacement::EPlantType::Unplanted;
+			UE::Anim::FootPlacement::EPlantType LastPlantType = UE::Anim::FootPlacement::EPlantType::Unplanted;
+			FPlane PlantPlaneWS = FPlane(FVector::UpVector, 0.0f);
+			FPlane PlantPlaneCS = FPlane(FVector::UpVector, 0.0f);
+			FQuat TwistCorrection = FQuat::Identity;
+			float TimeSinceFullyUnaligned = 0.0f;
+			bool bCanReachTarget = false;
+			// Whether we want to plant, independently from any dynamic pose adjustments we may do
+			bool bWantsToPlant = false;
+		} Plant;
+		
+		FTransform AlignedFootTransformWS = FTransform::Identity;
+		FTransform UnalignedFootTransformWS = FTransform::Identity;
+		FTransform AlignedFootTransformCS = FTransform::Identity;
+		FVector CachedIKToFKDir = FVector::UpVector;
+		
+		/* Interpolation */
+		struct FInterpolationData
+		{
+			FTransform UnalignedFootOffsetCS = FTransform::Identity;
+
+			FVectorSpringState PlantOffsetTranslationSpringState;
+			FQuaternionSpringState PlantOffsetRotationSpringState;
+
+			FVectorSpringState GroundTranslationSpringState;
+			FQuaternionSpringState GroundRotationSpringState;
+		} Interpolation;
+	};
+
+	struct FPlantRuntimeSettings
+	{
+		float MaxExtensionRatioSqrd = 0.0f;
+		float MinExtensionRatioSqrd = 0.0f;
+		float MaxLinearErrorSqrd = 0.0f;
+		float ReplantMaxLinearErrorSqrd = 0.0f;
+		float CosHalfMaxRotationError = 0.0f;
+		float CosHalfReplantMaxRotationError = 0.0f;
+	};
+
+	struct FPelvisRuntimeData
+	{
+		/* Bone IDs */
+		struct FBones
+		{
+			FCompactPoseBoneIndex FkBoneIndex = FCompactPoseBoneIndex(INDEX_NONE);
+			FCompactPoseBoneIndex IkBoneIndex = FCompactPoseBoneIndex(INDEX_NONE);
+		} Bones;
+
+		/* Settings-based properties */
+		float MaxOffsetHorizontalSqrd = 0.0f;
+		float MaxOffsetSqrd = 0.0f;
+
+		/* Input pose properties */
+		struct FInputPoseData
+		{
+			FTransform FKTransformCS = FTransform::Identity;
+			FTransform IKRootTransformCS = FTransform::Identity;
+		} InputPose;
+
+		/* Interpolation */
+		struct FInterpolationData
+		{
+			FVector PelvisTranslationOffset;
+			FVectorSpringState PelvisTranslationSpringState;
+		} Interpolation;
+	};
+
+	struct FCharacterData
+	{
+		FVector ComponentLocationWS = FVector::ZeroVector;
+		int NumFKPlanted = 0;
+		bool bIsOnGround = false;
+	};
+
+	struct FPlantResult
+	{
+	public:
+		FBoneTransform IkPlantTranformCS;
+		//FBoneTransform FkTipTransformCS;
+		//FBoneTransform FkHipTransformCS;
+	};
+
+#if ENABLE_ANIM_DEBUG
+	struct FDebugData
+	{
+		FVector OutputPelvisLocationWS = FVector::ZeroVector;
+		FVector InputPelvisLocationWS = FVector::ZeroVector;
+
+		TArray<FVector> OutputFootLocationsWS;
+		TArray<FVector> InputFootLocationsWS;
+
+		struct FLegExtension
+		{
+			float HyperExtensionAmount;
+			float RollAmount;
+			float PullAmount;
+		};
+		TArray<FLegExtension> LegsExtension;
+
+		void Init(const int32 InSize)
+		{
+			OutputFootLocationsWS.SetNumUninitialized(InSize);
+			InputFootLocationsWS.SetNumUninitialized(InSize);
+			LegsExtension.SetNumUninitialized(InSize);
+		}
+	};
+#endif
+
+	struct FEvaluationContext;
+}
+
+UENUM(BlueprintType)
+enum class EFootPlacementLockType : uint8
+{
+	Unlocked,
+	PivotAroundBall,
+	PivotAroundAnkle,
+	LockRotation
+};
+
+USTRUCT(BlueprintType)
+struct FFootPlacementInterpolationSettings
+{
+	GENERATED_BODY()
+
+public:
+	UPROPERTY(EditAnywhere, Category = "Plant Settings")
+	float UnplantLinearStiffness = 250.0f;
+
+	UPROPERTY(EditAnywhere, Category = "Plant Settings")
+	float UnplantLinearDamping = 1.0f;
+
+	UPROPERTY(EditAnywhere, Category = "Plant Settings")
+	float UnplantAngularStiffness = 450.0f;
+
+	UPROPERTY(EditAnywhere, Category = "Plant Settings")
+	float UnplantAngularDamping = 1.0f;
+
+	UPROPERTY(EditAnywhere, Category = "Plant Settings")
+	float FloorLinearStiffness = 1000.0f;
+
+	UPROPERTY(EditAnywhere, Category = "Plant Settings")
+	float FloorLinearDamping = 1.0f;
+
+	UPROPERTY(EditAnywhere, Category = "Plant Settings")
+	float FloorAngularStiffness = 450.0f;
+
+	UPROPERTY(EditAnywhere, Category = "Plant Settings")
+	float FloorAngularDamping = 1.0f;
+};
+
+USTRUCT(BlueprintType)
+struct FFootPlacementTraceSettings
+{
+	GENERATED_BODY()
+
+public:
+	// TODO: implement?
+	// Is tracing enabled?
+	UPROPERTY(EditAnywhere, Category = "Trace Settings")
+	bool bEnabled = true;
+
+	// A negative value extends the trace length above the bone
+	UPROPERTY(EditAnywhere, Category = "Trace Settings", meta=(EditCondition="bEnabled", DisplayAfter="bEnabled"))
+	float StartOffset = -75.0f;
+
+	// A positive value extends the trace length below the bone
+	UPROPERTY(EditAnywhere, Category = "Trace Settings", meta=(EditCondition="bEnabled", DisplayAfter="bEnabled"))
+	float EndOffset = 100.0f;
+
+	// The trace is a sphere sweep with this radius. It should be big enough to prevent the trace from going through 
+	// small geometry gaps
+	UPROPERTY(EditAnywhere, Category = "Trace Settings", meta=(EditCondition="bEnabled", DisplayAfter="bEnabled"))
+	float SweepRadius = 5.0f;
+
+	// The channel to use for our complex trace
+	UPROPERTY(EditAnywhere, Category = "Trace Settings", meta=(EditCondition="bEnabled", DisplayAfter="bEnabled"))
+	TEnumAsByte<ETraceTypeQuery> ComplexTraceChannel = TraceTypeQuery1;
+
+	// How much we align to simple vs complex collision when the foot is in flight
+	// Tracing against simple geometry (i.e. it's common for stairs to have simplified ramp collisions) can provide a 
+	// smoother trajectory when the foot is in flight
+	UPROPERTY(EditAnywhere, Category = "Trace Settings", meta = (EditCondition = "bEnabled", DisplayAfter = "bEnabled"))
+	float SimpleCollisionInfluence = 0.5f;
+
+	// The channel to use for our simple trace
+	UPROPERTY(EditAnywhere, Category = "Trace Settings", meta = (EditCondition = "bEnabled", DisplayAfter = "bEnabled"))
+	TEnumAsByte<ETraceTypeQuery> SimpleTraceChannel = TraceTypeQuery1;
+};
+
+USTRUCT()
+struct FFootPlacementRootDefinition
+{
+	GENERATED_BODY()
+
+public:
+	UPROPERTY(EditAnywhere, Category = "Pelvis Settings")
+	FBoneReference PelvisBone;
+
+	UPROPERTY(EditAnywhere, Category = "Pelvis Settings")
+	FBoneReference IKRootBone;
+
+public:
+	void Initialize(const FAnimationInitializeContext& Context);
+};
+
+
+USTRUCT(BlueprintType)
+struct FFootPlacementPelvisSettings
+{
+	GENERATED_BODY()
+
+public:
+	// Max horizontal foot adjustment we consider to lower the hips
+	// This can be used to prevent the hips from dropping too low when the feet are locked
+	// Exceeding this value will first attempt to roll the planted feet, and then slide
+	UPROPERTY(EditAnywhere, Category = "Pelvis Settings")
+	float MaxOffsetHorizontal = 20.0f;
+
+	// Max vertical offset from the input pose for the Pelvis.
+	// Reaching this limit means the feet may not reach their plant plane
+	UPROPERTY(EditAnywhere, Category = "Pelvis Settings")
+	float MaxOffset = 50.0f;
+
+	UPROPERTY(EditAnywhere, Category = "Pelvis Settings")
+	float LinearStiffness = 350.0f;
+
+	UPROPERTY(EditAnywhere, Category = "Pelvis Settings")
+	float LinearDamping = 1.0f;
+
+	// This is used to hold the Pelvis's interpolator in a fixed spot when the capsule suddenly moves (i.e. on a big step)
+	// If your camera is directly attached to the character with little to no smoothing, you may want this disabled
+	UPROPERTY(EditAnywhere, Category = "Pelvis Settings")
+	bool bCompensateForSuddenCapsuleMoves = true;
+};
+
+USTRUCT()
+struct FFootPlacemenLegDefinition
+{
+	GENERATED_BODY()
+
+public:
+	// Bone to be planted. For feet, use the heel/ankle joint.
+	UPROPERTY(EditAnywhere, Category = "Settings")
+	FBoneReference FKFootBone;
+
+	// TODO: can we optionally output as an attribute?
+	UPROPERTY(EditAnywhere, Category = "Settings")
+	FBoneReference IKFootBone;
+
+	// Secondary plant bone. For feet, use the ball joint.
+	UPROPERTY(EditAnywhere, Category = "Settings")
+	FBoneReference BallBone;
+
+	UPROPERTY(EditAnywhere, Category = "Settings")
+	int32 NumBonesInLimb = 2;
+
+	UPROPERTY(EditAnywhere, Category = "Settings")
+	FName SpeedCurveName = NAME_None;
+
+public:
+
+	void InitializeBoneReferences(const FBoneContainer& RequiredBones);
+};
+
+USTRUCT(BlueprintType)
+struct FFootPlacementPlantSettings
+{
+	GENERATED_BODY()
+
+public:
+	
+	// At this distance from the planting plane the bone is considered planted and will be fully aligned.
+	UPROPERTY(EditAnywhere, Category = "Plant Settings")
+	float FKPlantDistance = 10.0f;
+
+	// Bone is considered planted below this speed. Value is obtained from FKSpeedCurveName
+	UPROPERTY(EditAnywhere, Category = "Plant Settings")
+	float SpeedThreshold = 60.0f;
+
+	// Max extension ratio of the chain, calculated from the remaining length between current pose and full limb extension
+	UPROPERTY(EditAnywhere, Category = "Plant Settings")
+	float MaxExtensionRatio = 0.5f;
+
+	// Min extension ratio of the chain, calculated from the total limb length, and adjusted along the approach direction
+	UPROPERTY(EditAnywhere, Category = "Plant Settings")
+	float MinExtensionRatio = 0.2f;
+
+	// How much linear deviation causes the constraint to be released
+	UPROPERTY(EditAnywhere, Category = "Plant Settings")
+	float MaxLinearError = 35.0f;
+
+	// Below this value, proportional to MaxLinearError, the bone will replant
+	UPROPERTY(EditAnywhere, Category = "Plant Settings")
+	float ReplantMaxLinearErrorRatio = 0.35f;
+
+	UPROPERTY(EditAnywhere, Category = "Plant Settings")
+	EFootPlacementLockType LockType = EFootPlacementLockType::PivotAroundBall;
+	// Whether we fully lock rotation too.
+	UPROPERTY(EditAnywhere, Category = "Plant Settings")
+	bool bLockRotation = false;
+
+	// Whether we look the position of the foot along the tip bone (true) or the ik bone (false)
+	UPROPERTY(EditAnywhere, Category = "Plant Settings")
+	bool bPivotAroundBall = true;
+
+	// How much angular deviation (in degrees) causes the constraint to be released for replant
+	UPROPERTY(EditAnywhere, Category = "Plant Settings")
+	float MaxRotationError = 45.0;
+
+	// Below this value, proportional to MaxRotationError, the bone will replant
+	UPROPERTY(EditAnywhere, Category = "Plant Settings")
+	float ReplantMaxRotationErrorRatio = 0.5f;
+
+	// Speed at which we transition to fully unplanted.
+	// The range between SpeedThreshold and UnalignmentSpeedThreshold should roughly represent the roll-phase of the foot
+	// TODO: This feels innaccurate most of the time, and varies depending on anim speed. Improve this
+	UPROPERTY(EditAnywhere, Category = "Plant Settings")
+	float UnalignmentSpeedThreshold = 200.0f;
+
+	// Speed at which we transition to fully unplanted.
+	// The range between SpeedThreshold and UnalignmentSpeedThreshold should roughly represent the roll-phase of the foot
+	// TODO: This feels innaccurate most of the time, and varies depending on anim speed. Improve this
+	UPROPERTY(EditAnywhere, Category = "Plant Settings")
+	float AnkleTwistReduction = 0.75f;
+
+	// How much we can pull the the foot towards the hip to prevent hyperextension
+	// Before pulling the IK foot towards the FK foot
+	// While we're planted and within this threshold, the foot will roll instead of sliding
+	UPROPERTY(EditAnywhere, Category = "Settings")
+	float ExtensionPlantedPullOffset = 2.0f;
+
+	void Initialize(const FAnimationInitializeContext& Context);
+};
+
+
+USTRUCT(BlueprintInternalUseOnly, Experimental)
+struct ANIMATIONWARPINGRUNTIME_API FAnimNode_FootPlacement : public FAnimNode_SkeletalControlBase
+{
+	GENERATED_BODY()
+
+public:
+	UPROPERTY(EditAnywhere, Category = "Settings")
+	TArray<FFootPlacemenLegDefinition> LegDefinitions;
+
+	UPROPERTY(EditAnywhere, BlueprintReadWrite, Category = "Plant", meta = (PinHiddenByDefault))
+	FFootPlacementPlantSettings PlantSettings;
+
+	UPROPERTY(EditAnywhere, BlueprintReadWrite, Category = "Interpolation", meta = (PinHiddenByDefault))
+	FFootPlacementInterpolationSettings InterpolationSettings;
+
+	UPROPERTY(EditAnywhere, BlueprintReadWrite, Category = "Trace Settings", meta = (PinHiddenByDefault))
+	FFootPlacementTraceSettings TraceSettings;
+
+	// TODO: This wont work well when the animation doesn't have a single plant plane, i.e. a walking upstairs anim
+	UPROPERTY(EditAnywhere, Category = "Settings")
+	FBoneReference IKFootRootBone;
+
+	UPROPERTY(EditAnywhere, Category = "Settings")
+	FBoneReference PelvisBone;
+
+	UPROPERTY(EditAnywhere, BlueprintReadWrite, Category = "Plant Settings", meta = (PinHiddenByDefault))
+	FFootPlacementPelvisSettings PelvisSettings;
+
+public:
+	FAnimNode_FootPlacement();
+
+	// FAnimNode_Base interface
+	virtual void GatherDebugData(FNodeDebugData& DebugData) override;
+	// End of FAnimNode_Base interface
+
+	// FAnimNode_SkeletalControlBase interface
+	virtual void Initialize_AnyThread(const FAnimationInitializeContext& Context) override;
+	virtual void UpdateInternal(const FAnimationUpdateContext& Context) override;
+	virtual void EvaluateSkeletalControl_AnyThread(
+		FComponentSpacePoseContext& Output, 
+		TArray<FBoneTransform>& OutBoneTransforms) override;
+	virtual bool IsValidToEvaluate(const USkeleton* Skeleton, const FBoneContainer& RequiredBones) override;
+	// End of FAnimNode_SkeletalControlBase
+
+private:
+	// FAnimNode_SkeletalControlBase interface
+	virtual void InitializeBoneReferences(const FBoneContainer& RequiredBones) override;
+	// End of FAnimNode_SkeletalControlBase interface
+
+	void GatherPelvisDataFromInputs(const UE::Anim::FootPlacement::FEvaluationContext& Context);
+	void GatherLegDataFromInputs(
+		const UE::Anim::FootPlacement::FEvaluationContext& Context,
+		UE::Anim::FootPlacement::FLegRuntimeData& LegData,
+		const FFootPlacemenLegDefinition& LegDef);
+	void ProcessCharacterState(const UE::Anim::FootPlacement::FEvaluationContext& Context);
+	void ProcessFootAlignment(
+		const UE::Anim::FootPlacement::FEvaluationContext& Context,
+		UE::Anim::FootPlacement::FLegRuntimeData& LegData);
+
+	// Evaluates the Pelvis offsets required by the different planting bones and produces a result that best
+	// acommodates all of them
+	FTransform SolvePelvis(const UE::Anim::FootPlacement::FEvaluationContext& Context);
+
+	FTransform UpdatePelvisInterpolation(
+		const UE::Anim::FootPlacement::FEvaluationContext& Context,
+		const FTransform& TargetPelvisTransform);
+
+	// Post-processing adjustments + fix hyper-extension/compression
+	UE::Anim::FootPlacement::FPlantResult FinalizeFootAlignment(
+		const UE::Anim::FootPlacement::FEvaluationContext& Context,
+		UE::Anim::FootPlacement::FLegRuntimeData& LegData,
+		const FFootPlacemenLegDefinition& LegDef,
+		const FTransform& PelvisTransformCS);
+
+	FVector GetApproachDirWS(const FAnimationBaseContext& Context) const;
+
+private:
+	float CachedDeltaTime = 0.0f;
+	FVector LastComponentLocation = FVector::ZeroVector;
+
+	TArray<UE::Anim::FootPlacement::FLegRuntimeData> LegsData;
+	UE::Anim::FootPlacement::FPlantRuntimeSettings PlantRuntimeSettings;
+	UE::Anim::FootPlacement::FPelvisRuntimeData PelvisData;
+	UE::Anim::FootPlacement::FCharacterData CharacterData;
+
+	// Whether we want to plant, independently from any dynamic pose adjustments we may do
+	bool WantsToPlant(
+		const UE::Anim::FootPlacement::FEvaluationContext& Context,
+		const UE::Anim::FootPlacement::FLegRuntimeData::FInputPoseData& LegInputPose) const;
+
+	// Get Alignment Alpha based on current foot speed
+	// 0.0 is fully unaligned and the foot is in flight.
+	// 1.0 is fully aligned and the foot is planted.
+	float GetAlignmentAlpha(
+		const UE::Anim::FootPlacement::FEvaluationContext& Context,
+		const UE::Anim::FootPlacement::FLegRuntimeData::FInputPoseData& LegInputPose) const;
+
+	// This function looks at both the foot bone and the ball bone, returning the smallest distance to the
+	// planting plane. Note this distance can be negative, meaning it's penetrating.
+	float CalcTargetPlantPlaneDistance(
+		const UE::Anim::FootPlacement::FEvaluationContext& Context,
+		const UE::Anim::FootPlacement::FLegRuntimeData::FInputPoseData& LegInputPose) const;
+
+	// Find the horizontal pelvis offset range for the foot to reach:
+	// OutPelvisOffsetRangeCS.GetLowerBoundValue() = Offset for Desired extension
+	// OutPelvisOffsetRangeCS.GetUperBoundValue() = Offset for Max extension
+	void FindPelvisOffsetRangeForLimb(
+		const UE::Anim::FootPlacement::FEvaluationContext& Context,
+		const UE::Anim::FootPlacement::FLegRuntimeData::FInputPoseData& LegInputPoseData,
+		const FVector& PlantTargetLocationCS,
+		const FTransform& PelvisTransformCS,
+		const float LimbLength,
+		FFloatRange& OutPelvisOffsetRangeCS) const;
+
+	// Adjust LastPlantTransformWS to current, to have the foot pivot around the ball instead of the ankle
+	FTransform GetFootPivotAroundBallWS(
+		const UE::Anim::FootPlacement::FEvaluationContext& Context,
+		const UE::Anim::FootPlacement::FLegRuntimeData::FInputPoseData& LegInputPose,
+		const FTransform& LastPlantTransformWS) const;
+
+	// Align the transform the provided world space ground plant plane.
+	// Also outputs the twist along the ground plane needed to get there
+	void AlignPlantToGround(
+		const UE::Anim::FootPlacement::FEvaluationContext& Context,
+		const FPlane& PlantPlaneWS,
+		const UE::Anim::FootPlacement::FLegRuntimeData::FInputPoseData& LegInputPose,
+		FTransform& InOutFootTransformWS,
+		FQuat& OutTwistCorrection) const;
+
+	// Handles horizontal interpolation when unlocking the plant
+	FTransform UpdatePlantOffsetInterpolation(
+		const UE::Anim::FootPlacement::FEvaluationContext& Context,
+		UE::Anim::FootPlacement::FLegRuntimeData::FInterpolationData& InOutInterpData,
+		const FTransform& DesiredTransformCS) const;
+
+	// Handles the interpolation of the planting plane. Because the plant transform is specified with respect to the 
+	// planting plane, it cannot change abruptly without causing an animation pop. It must be interpolated instead.
+	void UpdatePlantingPlaneInterpolation(
+		const UE::Anim::FootPlacement::FEvaluationContext& Context,
+		const FTransform& FootTransformWS,
+		const FTransform& LastAlignedFootTransform,
+		const float AlignmentAlpha,
+		FPlane& InOutPlantPlane,
+		UE::Anim::FootPlacement::FLegRuntimeData::FInterpolationData& InOutInterpData) const;
+
+	// Checks unplanting and replanting conditions to determine if the foot is planted
+	void DeterminePlantType(
+		const UE::Anim::FootPlacement::FEvaluationContext& Context,
+		const FTransform& FKTransformWS,
+		const FTransform& CurrentBoneTransformWS,
+		UE::Anim::FootPlacement::FLegRuntimeData::FPlantData& InOutPlantData,
+		const UE::Anim::FootPlacement::FLegRuntimeData::FInputPoseData& LegInputPose) const;
+	
+	float GetMaxLimbExtension(const float DesiredExtension, const float LimbLength) const;
+	float GetMinLimbExtension(const float DesiredExtension, const float LimbLength) const;
+
+#if ENABLE_ANIM_DEBUG
+	UE::Anim::FootPlacement::FDebugData DebugData;
+	void DrawDebug(
+		const UE::Anim::FootPlacement::FEvaluationContext& Context,
+		const UE::Anim::FootPlacement::FLegRuntimeData& LegData,
+		const UE::Anim::FootPlacement::FPlantResult& PlantResult) const;
+#endif
+
+	bool bIsFirstUpdate = false;
+};
