Fix for CIS error in Foot Placement
#rnx

diff --git a/Engine/Plugins/Animation/AnimationWarping/Source/Runtime/Private/BoneControllers/AnimNode_FootPlacement.cpp b/Engine/Plugins/Animation/AnimationWarping/Source/Runtime/Private/BoneControllers/AnimNode_FootPlacement.cpp
--- a/Engine/Plugins/Animation/AnimationWarping/Source/Runtime/Private/BoneControllers/AnimNode_FootPlacement.cpp
+++ b/Engine/Plugins/Animation/AnimationWarping/Source/Runtime/Private/BoneControllers/AnimNode_FootPlacement.cpp
@@ -937,7 +937,10 @@ void FAnimNode_FootPlacement::EvaluateSkeletalControl_AnyThread(FComponentSpaceP
 
 	check(OutBoneTransforms.Num() == 0);
 
+#if ENABLE_ANIM_DEBUG
 	UE::Anim::FootPlacement::FDebugData LastDebugData = DebugData;
+#endif
+
 	// TODO: Support a different approach direction
 	const FVector ApproachDirCS = -FVector::UpVector;
 	UE::Anim::FootPlacement::FEvaluationContext FootPlacementContext(Output, ApproachDirCS, CachedDeltaTime);
