Foot Placement: Change node color to match other pose warping nodes
#rnx

diff --git a/Engine/Plugins/Animation/AnimationWarping/Source/Editor/Private/AnimGraph/AnimGraphNode_FootPlacement.cpp b/Engine/Plugins/Animation/AnimationWarping/Source/Editor/Private/AnimGraph/AnimGraphNode_FootPlacement.cpp
--- a/Engine/Plugins/Animation/AnimationWarping/Source/Editor/Private/AnimGraph/AnimGraphNode_FootPlacement.cpp
+++ b/Engine/Plugins/Animation/AnimationWarping/Source/Editor/Private/AnimGraph/AnimGraphNode_FootPlacement.cpp
@@ -25,6 +25,11 @@ FText UAnimGraphNode_FootPlacement::GetTooltipText() const
 	return LOCTEXT("FootPlacementTooltip", "Foot Placement.");
 }
 
+FLinearColor UAnimGraphNode_FootPlacement::GetNodeTitleColor() const
+{
+	return FLinearColor(FColor(153.f, 0.f, 0.f));
+}
+
 FText UAnimGraphNode_FootPlacement::GetNodeTitle(ENodeTitleType::Type TitleType) const
 {
 	return GetControllerDescription();
diff --git a/Engine/Plugins/Animation/AnimationWarping/Source/Editor/Public/AnimGraph/AnimGraphNode_FootPlacement.h b/Engine/Plugins/Animation/AnimationWarping/Source/Editor/Public/AnimGraph/AnimGraphNode_FootPlacement.h
--- a/Engine/Plugins/Animation/AnimationWarping/Source/Editor/Public/AnimGraph/AnimGraphNode_FootPlacement.h
+++ b/Engine/Plugins/Animation/AnimationWarping/Source/Editor/Public/AnimGraph/AnimGraphNode_FootPlacement.h
@@ -17,6 +17,7 @@ public:
 	// UEdGraphNode interface
 	virtual FText GetNodeTitle(ENodeTitleType::Type TitleType) const override;
 	virtual FText GetTooltipText() const override;
+	virtual FLinearColor GetNodeTitleColor() const override;
 	// End of UEdGraphNode interface
 
 protected:
