Foot Placement
-Change behavior of  trace. When the surface is unwalkable, still use the hit position, but don't align to it.
-Updated UPROPERTY categories
#rnx

diff --git a/Engine/Plugins/Animation/AnimationWarping/Source/Runtime/Private/BoneControllers/AnimNode_FootPlacement.cpp b/Engine/Plugins/Animation/AnimationWarping/Source/Runtime/Private/BoneControllers/AnimNode_FootPlacement.cpp
--- a/Engine/Plugins/Animation/AnimationWarping/Source/Runtime/Private/BoneControllers/AnimNode_FootPlacement.cpp
+++ b/Engine/Plugins/Animation/AnimationWarping/Source/Runtime/Private/BoneControllers/AnimNode_FootPlacement.cpp
@@ -128,14 +128,15 @@ namespace UE::Anim::FootPlacement
 		const bool bHit = Context.World->SweepSingleByChannel(
 			HitResult, TraceStart, TraceEnd, FQuat::Identity, CollisionChannel, CollisionShape, QueryParams);
 
+		OutImpactLocationWS = HitResult.ImpactPoint;
 		if (!Context.GetMovementComponentIsWalkable(HitResult))
 		{
+			// If the surface hit isn't walkable, use the negated trace direction as the impact normal
+			OutImpactNormalWS = -TraceDirectionWS;
 			return false;
 		}
 
-		OutImpactLocationWS = HitResult.ImpactPoint;
 		OutImpactNormalWS = HitResult.ImpactNormal;
-
 		return true;
 	}
 
diff --git a/Engine/Plugins/Animation/AnimationWarping/Source/Runtime/Public/BoneControllers/AnimNode_FootPlacement.h b/Engine/Plugins/Animation/AnimationWarping/Source/Runtime/Public/BoneControllers/AnimNode_FootPlacement.h
--- a/Engine/Plugins/Animation/AnimationWarping/Source/Runtime/Public/BoneControllers/AnimNode_FootPlacement.h
+++ b/Engine/Plugins/Animation/AnimationWarping/Source/Runtime/Public/BoneControllers/AnimNode_FootPlacement.h
@@ -406,17 +406,17 @@ public:
 	UPROPERTY(EditAnywhere, BlueprintReadWrite, Category = "Interpolation", meta = (PinHiddenByDefault))
 	FFootPlacementInterpolationSettings InterpolationSettings;
 
-	UPROPERTY(EditAnywhere, BlueprintReadWrite, Category = "Trace Settings", meta = (PinHiddenByDefault))
+	UPROPERTY(EditAnywhere, BlueprintReadWrite, Category = "Trace", meta = (PinHiddenByDefault))
 	FFootPlacementTraceSettings TraceSettings;
 
 	// TODO: This wont work well when the animation doesn't have a single plant plane, i.e. a walking upstairs anim
-	UPROPERTY(EditAnywhere, Category = "Settings")
+	UPROPERTY(EditAnywhere, Category = "Plant")
 	FBoneReference IKFootRootBone;
 
-	UPROPERTY(EditAnywhere, Category = "Settings")
+	UPROPERTY(EditAnywhere, Category = "Pelvis")
 	FBoneReference PelvisBone;
 
-	UPROPERTY(EditAnywhere, BlueprintReadWrite, Category = "Plant Settings", meta = (PinHiddenByDefault))
+	UPROPERTY(EditAnywhere, BlueprintReadWrite, Category = "Pelvis", meta = (PinHiddenByDefault))
 	FFootPlacementPelvisSettings PelvisSettings;
 
 public:
