Offset Root Bone - Added experimental node

diff --git a/Engine/Plugins/Animation/AnimationWarping/Source/Editor/Private/AnimGraph/AnimGraphNode_OffsetRootBone.cpp b/Engine/Plugins/Animation/AnimationWarping/Source/Editor/Private/AnimGraph/AnimGraphNode_OffsetRootBone.cpp
new file mode 100644
--- /dev/null
+++ b/Engine/Plugins/Animation/AnimationWarping/Source/Editor/Private/AnimGraph/AnimGraphNode_OffsetRootBone.cpp
@@ -0,0 +1,121 @@
+// Copyright Epic Games, Inc. All Rights Reserved.
+
+#include "AnimGraph/AnimGraphNode_OffsetRootBone.h"
+#include "Animation/AnimRootMotionProvider.h"
+#include "Kismet2/BlueprintEditorUtils.h"
+#include "DetailCategoryBuilder.h"
+#include "DetailLayoutBuilder.h"
+#include "ScopedTransaction.h"
+
+#define LOCTEXT_NAMESPACE "AnimationWarping"
+
+UAnimGraphNode_OffsetRootBone::UAnimGraphNode_OffsetRootBone(const FObjectInitializer& ObjectInitializer) : Super(ObjectInitializer)
+{
+}
+
+FText UAnimGraphNode_OffsetRootBone::GetControllerDescription() const
+{
+	return LOCTEXT("OffsetRootBone", "Offset Root Bone");
+}
+
+FText UAnimGraphNode_OffsetRootBone::GetTooltipText() const
+{
+	return LOCTEXT("OffsetRootBoneTooltip", "Applies an offset to the root bone");
+}
+
+FText UAnimGraphNode_OffsetRootBone::GetNodeTitle(ENodeTitleType::Type TitleType) const
+{
+	return GetControllerDescription();
+}
+
+FLinearColor UAnimGraphNode_OffsetRootBone::GetNodeTitleColor() const
+{
+	return FLinearColor(FColor(153.f, 0.f, 0.f));
+}
+
+void UAnimGraphNode_OffsetRootBone::CustomizePinData(UEdGraphPin* Pin, FName SourcePropertyName, int32 ArrayIndex) const
+{
+	Super::CustomizePinData(Pin, SourcePropertyName, ArrayIndex);
+
+	if (Pin->PinName == GET_MEMBER_NAME_STRING_CHECKED(FAnimNode_OffsetRootBone, TranslationDelta))
+	{
+		Pin->bHidden = (Node.EvaluationMode == EWarpingEvaluationMode::Graph);
+	}
+
+	if (Pin->PinName == GET_MEMBER_NAME_STRING_CHECKED(FAnimNode_OffsetRootBone, RotationDelta))
+	{
+		Pin->bHidden = (Node.EvaluationMode == EWarpingEvaluationMode::Graph);
+	}
+}
+
+void UAnimGraphNode_OffsetRootBone::CustomizeDetails(IDetailLayoutBuilder& DetailBuilder)
+{
+	DECLARE_SCOPE_HIERARCHICAL_COUNTER_FUNC()
+	Super::CustomizeDetails(DetailBuilder);
+
+	TSharedRef<IPropertyHandle> NodeHandle = DetailBuilder.GetProperty(FName(TEXT("Node")), GetClass());
+
+	if (Node.EvaluationMode == EWarpingEvaluationMode::Graph)
+	{
+		DetailBuilder.HideProperty(NodeHandle->GetChildHandle(GET_MEMBER_NAME_CHECKED(FAnimNode_OffsetRootBone, TranslationDelta)));
+		DetailBuilder.HideProperty(NodeHandle->GetChildHandle(GET_MEMBER_NAME_CHECKED(FAnimNode_OffsetRootBone, RotationDelta)));
+	}
+}
+
+void UAnimGraphNode_OffsetRootBone::PostEditChangeProperty(struct FPropertyChangedEvent& PropertyChangedEvent)
+{
+	DECLARE_SCOPE_HIERARCHICAL_COUNTER_FUNC()
+	Super::PostEditChangeProperty(PropertyChangedEvent);
+
+	bool bRequiresNodeReconstruct = false;
+	FProperty* ChangedProperty = PropertyChangedEvent.Property;
+
+	if (ChangedProperty)
+	{
+		if ((ChangedProperty->GetFName() == GET_MEMBER_NAME_STRING_CHECKED(FAnimNode_OffsetRootBone, EvaluationMode)))
+		{
+			FScopedTransaction Transaction(LOCTEXT("OffsetRootBone_ChangeEvaluationMode", "Offset Root Bone Change Evaluation Mode"));
+			Modify();
+
+			// Break links to pins going away
+			for (int32 PinIndex = 0; PinIndex < Pins.Num(); ++PinIndex)
+			{
+				UEdGraphPin* Pin = Pins[PinIndex];
+				if (Pin->PinName == GET_MEMBER_NAME_STRING_CHECKED(FAnimNode_OffsetRootBone, TranslationDelta))
+				{
+					if (Node.EvaluationMode == EWarpingEvaluationMode::Graph)
+					{
+						Pin->BreakAllPinLinks();
+					}
+				}
+				else if (Pin->PinName == GET_MEMBER_NAME_STRING_CHECKED(FAnimNode_OffsetRootBone, RotationDelta))
+				{
+					if (Node.EvaluationMode == EWarpingEvaluationMode::Graph)
+					{
+						Pin->BreakAllPinLinks();
+					}
+				}
+			}
+
+			bRequiresNodeReconstruct = true;
+		}
+	}
+
+	if (bRequiresNodeReconstruct)
+	{
+		ReconstructNode();
+		FBlueprintEditorUtils::MarkBlueprintAsStructurallyModified(GetBlueprint());
+	}
+}
+
+void UAnimGraphNode_OffsetRootBone::GetInputLinkAttributes(FNodeAttributeArray& OutAttributes) const
+{
+	OutAttributes.Add(UE::Anim::IAnimRootMotionProvider::AttributeName);
+}
+
+void UAnimGraphNode_OffsetRootBone::GetOutputLinkAttributes(FNodeAttributeArray& OutAttributes) const
+{
+	OutAttributes.Add(UE::Anim::IAnimRootMotionProvider::AttributeName);
+}
+
+#undef LOCTEXT_NAMESPACE
diff --git a/Engine/Plugins/Animation/AnimationWarping/Source/Editor/Public/AnimGraph/AnimGraphNode_OffsetRootBone.h b/Engine/Plugins/Animation/AnimationWarping/Source/Editor/Public/AnimGraph/AnimGraphNode_OffsetRootBone.h
new file mode 100644
--- /dev/null
+++ b/Engine/Plugins/Animation/AnimationWarping/Source/Editor/Public/AnimGraph/AnimGraphNode_OffsetRootBone.h
@@ -0,0 +1,38 @@
+// Copyright Epic Games, Inc. All Rights Reserved.
+
+#pragma once
+
+#include "AnimGraphNode_SkeletalControlBase.h"
+#include "BoneControllers/AnimNode_OffsetRootBone.h"
+
+#include "AnimGraphNode_OffsetRootBone.generated.h"
+
+UCLASS(Experimental)
+class ANIMATIONWARPINGEDITOR_API UAnimGraphNode_OffsetRootBone : public UAnimGraphNode_SkeletalControlBase
+{
+	GENERATED_UCLASS_BODY()
+
+	UPROPERTY(EditAnywhere, Category = Settings)
+	FAnimNode_OffsetRootBone Node;
+
+public:
+	// UEdGraphNode interface
+	virtual FText GetNodeTitle(ENodeTitleType::Type TitleType) const override;
+	virtual FLinearColor GetNodeTitleColor() const override;
+	virtual FText GetTooltipText() const override;
+	virtual void CustomizeDetails(IDetailLayoutBuilder& DetailBuilder) override;
+	virtual void PostEditChangeProperty(struct FPropertyChangedEvent& PropertyChangedEvent) override;
+	virtual void GetInputLinkAttributes(FNodeAttributeArray& OutAttributes) const override;
+	virtual void GetOutputLinkAttributes(FNodeAttributeArray& OutAttributes) const override;
+	// End of UEdGraphNode interface
+
+protected:
+	// UAnimGraphNode_Base interface
+	virtual void CustomizePinData(UEdGraphPin* Pin, FName SourcePropertyName, int32 ArrayIndex) const override;
+	// End of UAnimGraphNode_Base interface
+
+	// UAnimGraphNode_SkeletalControlBase interface
+	virtual FText GetControllerDescription() const override;
+	virtual const FAnimNode_SkeletalControlBase* GetNode() const override { return &Node; }
+	// End of UAnimGraphNode_SkeletalControlBase interface
+};
diff --git a/Engine/Plugins/Animation/AnimationWarping/Source/Runtime/Private/BoneControllers/AnimNode_OffsetRootBone.cpp b/Engine/Plugins/Animation/AnimationWarping/Source/Runtime/Private/BoneControllers/AnimNode_OffsetRootBone.cpp
new file mode 100644
--- /dev/null
+++ b/Engine/Plugins/Animation/AnimationWarping/Source/Runtime/Private/BoneControllers/AnimNode_OffsetRootBone.cpp
@@ -0,0 +1,306 @@
+// Copyright Epic Games, Inc. All Rights Reserved.
+
+#include "BoneControllers/AnimNode_OffsetRootBone.h"
+#include "Animation/AnimInstanceProxy.h"
+#include "Animation/AnimRootMotionProvider.h"
+#include "Math/UnrealMathUtility.h"
+
+DECLARE_CYCLE_STAT(TEXT("OffsetRootBone Eval"), STAT_OffsetRootBone_Eval, STATGROUP_Anim);
+
+#if ENABLE_ANIM_DEBUG
+TAutoConsoleVariable<int32> CVarAnimNodeOffsetRootBoneDebug(TEXT("a.AnimNode.OffsetRootBone.Debug"), 0, TEXT("Turn on visualization debugging for Offset Root Bone"));
+TAutoConsoleVariable<int32> CVarAnimNodeOffsetRootBoneEnable(TEXT("a.AnimNode.OffsetRootBone.Enable"), 1, TEXT("Toggle Offset Root Bone"));
+TAutoConsoleVariable<int32> CVarAnimNodeOffsetRootBoneModifyBone(TEXT("a.AnimNode.OffsetRootBone.ModifyBone"), 1, TEXT("Toggle whether the transform is applied to the bone"));
+#endif
+
+
+namespace UE::Anim::OffsetRootBone
+{
+	// Taken from https://theorangeduck.com/page/spring-roll-call#implicitspringdamper
+	static float DamperImplicit(float Halflife, float DeltaTime, float Epsilon = 1e-8f)
+	{
+		return (1.0f - FMath::InvExpApprox((0.69314718056f * DeltaTime) / (Halflife + Epsilon)));
+	}
+}
+
+// TODO: implement 
+void FAnimNode_OffsetRootBone::GatherDebugData(FNodeDebugData& DebugData)
+{
+}
+
+void FAnimNode_OffsetRootBone::Initialize_AnyThread(const FAnimationInitializeContext& Context)
+{
+	FAnimNode_SkeletalControlBase::Initialize_AnyThread(Context);
+	AnimInstanceProxy = Context.AnimInstanceProxy;
+	ComponentTransform = Context.AnimInstanceProxy->GetComponentTransform();
+	SimulatedTranslation = ComponentTransform.GetLocation();
+	SimulatedRotation = ComponentTransform.GetRotation();
+}
+
+void FAnimNode_OffsetRootBone::UpdateInternal(const FAnimationUpdateContext& Context)
+{
+	FAnimNode_SkeletalControlBase::UpdateInternal(Context);
+	CachedDeltaTime = Context.GetDeltaTime();
+}
+
+void FAnimNode_OffsetRootBone::EvaluateSkeletalControl_AnyThread(FComponentSpacePoseContext& Output, TArray<FBoneTransform>& OutBoneTransforms)
+{
+	SCOPE_CYCLE_COUNTER(STAT_OffsetRootBone_Eval);
+	check(OutBoneTransforms.Num() == 0);
+
+	bool bGraphDriven = false;
+	const UE::Anim::IAnimRootMotionProvider* RootMotionProvider = UE::Anim::IAnimRootMotionProvider::Get();
+
+	if (GetEvaluationMode() == EWarpingEvaluationMode::Graph)
+	{
+		bGraphDriven = RootMotionProvider != nullptr;
+		ensureMsgf(bGraphDriven, TEXT("Graph driven Offset Root Bone expected a valid root motion delta provider interface."));
+	}
+
+	const FCompactPoseBoneIndex TargetBoneIndex(0);
+	const FTransform InputBoneTransform = Output.Pose.GetComponentSpaceTransform(TargetBoneIndex);
+
+	const FTransform LastComponentTransform = ComponentTransform;
+	ComponentTransform = AnimInstanceProxy->GetComponentTransform();
+
+	const bool bShouldConsumeTranslationOffset = GetOffsetSettings().TranslationMode == EOffsetRootBoneMode::Accumulate ||
+										GetOffsetSettings().TranslationMode == EOffsetRootBoneMode::Interpolate;
+	const bool bShouldConsumeRotationOffset = GetOffsetSettings().RotationMode == EOffsetRootBoneMode::Accumulate ||
+										GetOffsetSettings().RotationMode == EOffsetRootBoneMode::Interpolate;
+
+	FTransform RootMotionTransformDelta = FTransform::Identity;
+	if (bGraphDriven)
+	{
+		// Graph driven mode will override translation and rotation delta
+		// with the intent of the current animation sub-graph's accumulated root motion
+		bGraphDriven = RootMotionProvider->ExtractRootMotion(Output.CustomAttributes, RootMotionTransformDelta);
+	}
+	else
+	{
+		// Apply the offset as is (component space) in manual mode
+		RootMotionTransformDelta = FTransform(GetRotationDelta(), GetTranslationDelta());
+	}
+
+	FTransform ConsumedRootMotionDelta;
+
+	if (bShouldConsumeTranslationOffset)
+	{
+		// Grab root motion translation from the root motion attribute
+		ConsumedRootMotionDelta.SetTranslation(RootMotionTransformDelta.GetTranslation());
+	}
+	else
+	{
+		// Accumulate the translation frame delta into the simulated translation, to keep it in place
+		const FVector ComponentTranslationDelta = ComponentTransform.GetLocation() - LastComponentTransform.GetLocation();
+		SimulatedTranslation += ComponentTranslationDelta;
+	}
+
+	if (bShouldConsumeRotationOffset)
+	{
+		// Grab root motion rotation from the root motion attribute
+		ConsumedRootMotionDelta.SetRotation(RootMotionTransformDelta.GetRotation());
+	}
+	else
+	{
+		// Accumulate the rotation frame delta into the simulated translation, to keep it in place
+		const FQuat ComponentRotationDelta = LastComponentTransform.GetRotation().Inverse() * ComponentTransform.GetRotation();
+		SimulatedRotation = ComponentRotationDelta * SimulatedRotation;
+	}
+
+	FTransform SimulatedTransform(SimulatedRotation, SimulatedTranslation);
+	// Apply the root motion delta
+	SimulatedTransform = ConsumedRootMotionDelta * SimulatedTransform;
+
+	SimulatedTranslation = SimulatedTransform.GetLocation();
+	SimulatedRotation = SimulatedTransform.GetRotation();
+	
+	// TODO: Make this a parameter
+	const FVector GravityDirCS = -FVector::UpVector;
+	// Simulated translation should stay the same along the approach direction
+	SimulatedTranslation = FVector::PointPlaneProject(SimulatedTranslation, ComponentTransform.GetLocation(), GravityDirCS);
+
+	const FBoneContainer& RequiredBones = Output.Pose.GetPose().GetBoneContainer();
+
+	bool bModifyBone = true;
+#if ENABLE_ANIM_DEBUG
+	bModifyBone = CVarAnimNodeOffsetRootBoneModifyBone.GetValueOnAnyThread() == 1;
+#endif
+
+	if (GetOffsetSettings().TranslationMode == EOffsetRootBoneMode::Release ||
+		GetOffsetSettings().TranslationMode == EOffsetRootBoneMode::Interpolate)
+	{
+		FVector TranslationOffset = ComponentTransform.GetLocation() - SimulatedTranslation;
+
+		const float DampenAlpha = UE::Anim::OffsetRootBone::DamperImplicit(GetOffsetSettings().TranslationHalflife, CachedDeltaTime);
+		FVector TranslationOffsetDelta = FMath::Lerp(FVector::ZeroVector, TranslationOffset, DampenAlpha);
+
+		if (GetOffsetSettings().bClampToTranslationVelocity)
+		{
+			const float RootMotionDelta = RootMotionTransformDelta.GetLocation().Size();
+			const float MaxDelta = 	GetOffsetSettings().TranslationSpeedRatio * RootMotionDelta;
+
+			const float AdjustmentDelta = TranslationOffsetDelta.Size();
+			if (AdjustmentDelta > MaxDelta)
+			{
+				TranslationOffsetDelta = MaxDelta * TranslationOffsetDelta.GetSafeNormal2D();
+			}
+		}
+
+		SimulatedTranslation = SimulatedTranslation + TranslationOffsetDelta;
+	}
+
+	if (GetOffsetSettings().RotationMode == EOffsetRootBoneMode::Release ||
+		GetOffsetSettings().RotationMode == EOffsetRootBoneMode::Interpolate)
+	{
+		FQuat RotationOffset = ComponentTransform.GetRotation() * SimulatedRotation.Inverse();
+		RotationOffset.Normalize();
+		RotationOffset = RotationOffset.W < 0.0 ? -RotationOffset : RotationOffset;
+
+		const float DampenAlpha = UE::Anim::OffsetRootBone::DamperImplicit(GetOffsetSettings().RotationHalfLife, CachedDeltaTime);
+		FQuat RotationOffsetDelta = FQuat::Slerp(FQuat::Identity, RotationOffset, DampenAlpha);
+
+		if (GetOffsetSettings().bClampToRotationVelocity)
+		{
+			float RotationMotionAngleDelta;
+			FVector RootMotionRotationAxis;
+			RootMotionTransformDelta.GetRotation().ToAxisAndAngle(RootMotionRotationAxis, RotationMotionAngleDelta);
+
+			float MaxRotationAngle = GetOffsetSettings().RotationSpeedRatio * RotationMotionAngleDelta;
+
+			FVector DeltaAxis;
+			float DeltaAngle;
+			RotationOffsetDelta.ToAxisAndAngle(DeltaAxis, DeltaAngle);
+
+			if (DeltaAngle > MaxRotationAngle)
+			{
+				RotationOffsetDelta = FQuat(DeltaAxis, MaxRotationAngle);
+			}
+		}
+
+		SimulatedRotation = RotationOffsetDelta * SimulatedRotation;
+	}
+
+	if (GetOffsetSettings().MaxTranslationError >= 0.0f)
+	{
+		FVector TranslationOffset = ComponentTransform.GetLocation() - SimulatedTranslation;
+		const float TranslationOffsetSize = TranslationOffset.Size();
+		if (TranslationOffsetSize > GetOffsetSettings().MaxTranslationError)
+		{
+			TranslationOffset = TranslationOffset.GetClampedToMaxSize(GetOffsetSettings().MaxTranslationError);
+			SimulatedTranslation = ComponentTransform.GetLocation() - TranslationOffset;
+		}
+	}
+
+	const float MaxAngleRadians = FMath::DegreesToRadians(GetOffsetSettings().MaxRotationError);
+	if (GetOffsetSettings().MaxRotationError >= 0.0f)
+	{
+		FQuat RotationOffset = ComponentTransform.GetRotation().Inverse() * SimulatedRotation;
+		RotationOffset.Normalize();
+		RotationOffset = RotationOffset.W < 0.0 ? -RotationOffset : RotationOffset;
+
+		FVector OffsetAxis;
+		float OffsetAngle;
+		RotationOffset.ToAxisAndAngle(OffsetAxis, OffsetAngle);
+
+		if (FMath::Abs(OffsetAngle) > MaxAngleRadians)
+		{
+			RotationOffset = FQuat(OffsetAxis, MaxAngleRadians);
+			SimulatedRotation = RotationOffset * ComponentTransform.GetRotation();
+		}
+	}
+
+	// Apply the offset adjustments to the simulated transform
+	SimulatedTransform.SetLocation(SimulatedTranslation);
+	SimulatedTransform.SetRotation(SimulatedRotation);
+
+	// Combine with the input pose's bone transform, to preserve any adjustments done before this node in the graph
+	FTransform TargetBoneTransform = SimulatedTransform * ComponentTransform.Inverse();
+	// Accumulate the input bone transform to keep the offset independent from any previous adjustments to the root
+	TargetBoneTransform.Accumulate(InputBoneTransform);
+	if (bModifyBone)
+	{
+		OutBoneTransforms.Add(FBoneTransform(TargetBoneIndex, TargetBoneTransform));
+	}
+
+#if ENABLE_ANIM_DEBUG
+	bool bDebugging = CVarAnimNodeOffsetRootBoneDebug.GetValueOnAnyThread() == 1;
+	if (bDebugging)
+	{
+		const float InnerCircleRadius = 40.0f;
+		const float CircleThickness = 0.5f;
+		const float ConeThickness = 0.2f;
+
+		const FTransform TargetBoneInitialTransformWorld = InputBoneTransform * ComponentTransform;
+		const FTransform TargetBoneTransformWorld = TargetBoneTransform * ComponentTransform;
+
+		if (GetOffsetSettings().MaxTranslationError >= 0.0f)
+		{
+			const float OuterCircleRadius = GetOffsetSettings().MaxTranslationError + InnerCircleRadius;
+			AnimInstanceProxy->AnimDrawDebugCircle(ComponentTransform.GetLocation(), OuterCircleRadius, 36, FColor::Red,
+				FVector::UpVector, false, -1.0f, SDPG_World, CircleThickness);
+		}
+
+		AnimInstanceProxy->AnimDrawDebugCircle(ComponentTransform.GetLocation(), InnerCircleRadius, 36, FColor::Blue,
+			FVector::UpVector, false, -1.0f, SDPG_World, CircleThickness);
+			
+		AnimInstanceProxy->AnimDrawDebugCircle(TargetBoneTransformWorld.GetLocation(), InnerCircleRadius, 36, FColor::Green,
+			FVector::UpVector, false, -1.0f, SDPG_World, CircleThickness);
+
+		const int32 ConeSegments =  FMath::RoundUpToPowerOfTwo((GetOffsetSettings().MaxRotationError / 180.0f) * 12 );
+		const FVector ArcDirection = ComponentTransform.GetRotation().GetRightVector();
+		AnimInstanceProxy->AnimDrawDebugCone(TargetBoneTransformWorld.GetLocation(), 0.9f * InnerCircleRadius, ArcDirection,
+			MaxAngleRadians, 0.0f, ConeSegments, FColor::Red, false, -1.0f, SDPG_World, ConeThickness);
+
+		AnimInstanceProxy->AnimDrawDebugDirectionalArrow(
+			TargetBoneTransformWorld.GetLocation() + InnerCircleRadius * TargetBoneInitialTransformWorld.GetRotation().GetRightVector(),
+			TargetBoneTransformWorld.GetLocation() + 1.5f * InnerCircleRadius * TargetBoneInitialTransformWorld.GetRotation().GetRightVector(),
+			40.f, FColor::Red, false, 0.f, CircleThickness);
+
+		AnimInstanceProxy->AnimDrawDebugDirectionalArrow(
+			TargetBoneTransformWorld.GetLocation() + InnerCircleRadius * TargetBoneTransformWorld.GetRotation().GetRightVector(),
+			TargetBoneTransformWorld.GetLocation() + 1.3f * InnerCircleRadius * TargetBoneTransformWorld.GetRotation().GetRightVector(),
+			40.f, FColor::Blue, false, 0.f, CircleThickness);
+	}
+#endif
+
+
+	if (bGraphDriven && bModifyBone)
+	{
+		const FTransform RemainingRootMotionDelta = ConsumedRootMotionDelta * RootMotionTransformDelta.Inverse();
+		FTransform TargetRootMotionTransformDelta(RemainingRootMotionDelta.GetRotation(), RemainingRootMotionDelta.GetTranslation(), RootMotionTransformDelta.GetScale3D());
+		const bool bRootMotionOverridden = RootMotionProvider->OverrideRootMotion(RootMotionTransformDelta, Output.CustomAttributes);
+		ensure(bRootMotionOverridden);
+	}
+}
+
+bool FAnimNode_OffsetRootBone::IsValidToEvaluate(const USkeleton* Skeleton, const FBoneContainer& RequiredBones)
+{
+#if ENABLE_ANIM_DEBUG
+	if (CVarAnimNodeOffsetRootBoneEnable.GetValueOnAnyThread() == 0)
+	{
+		return false;
+	}
+#endif
+
+	return true;
+}
+
+EWarpingEvaluationMode FAnimNode_OffsetRootBone::GetEvaluationMode() const
+{
+	return GET_ANIM_NODE_DATA(EWarpingEvaluationMode, EvaluationMode);
+}
+
+const FOffsetRootBoneSettings& FAnimNode_OffsetRootBone::GetOffsetSettings() const
+{
+	return GET_ANIM_NODE_DATA(FOffsetRootBoneSettings, OffsetSettings);
+}
+
+const FVector& FAnimNode_OffsetRootBone::GetTranslationDelta() const
+{
+	return GET_ANIM_NODE_DATA(FVector, TranslationDelta);
+}
+
+const FRotator& FAnimNode_OffsetRootBone::GetRotationDelta() const
+{
+	return GET_ANIM_NODE_DATA(FRotator, RotationDelta);
+}
diff --git a/Engine/Plugins/Animation/AnimationWarping/Source/Runtime/Public/BoneControllers/AnimNode_OffsetRootBone.h b/Engine/Plugins/Animation/AnimationWarping/Source/Runtime/Public/BoneControllers/AnimNode_OffsetRootBone.h
new file mode 100644
--- /dev/null
+++ b/Engine/Plugins/Animation/AnimationWarping/Source/Runtime/Public/BoneControllers/AnimNode_OffsetRootBone.h
@@ -0,0 +1,140 @@
+// Copyright Epic Games, Inc. All Rights Reserved.
+
+#pragma once
+
+#include "BoneControllers/BoneControllerTypes.h"
+#include "BoneControllers/BoneControllerSolvers.h"
+#include "BoneControllers/AnimNode_SkeletalControlBase.h"
+#include "AnimNode_OffsetRootBone.generated.h"
+
+UENUM(BlueprintType)
+enum class EOffsetRootBoneMode : uint8
+{
+	// Accumulate the mesh component's movement into the offset.
+	// In this mode, if the mesh component moves 
+	// the offset will counter the motion, and the root will stay in place
+	Accumulate,
+	// Continuously interpolate the offset out
+	// In this mode, if the mesh component moves
+	// The root will stay behind, but will attempt to catch up
+	Interpolate,
+	// Stops accumulating the mesh component's movement delta into the root offset
+	// In this mode, whatever offset we have will be conserved, 
+	// but we won't accumulate any more
+	Hold,
+	// Release the offset and stop accumulating the mesh component's movement delta.
+	// In this mode we will "blend out" the offset
+	Release,
+};
+
+USTRUCT(BlueprintType)
+struct FOffsetRootBoneSettings
+{
+	GENERATED_BODY()
+
+	// The translation offset behavior mode
+	UPROPERTY(EditAnywhere, BlueprintReadWrite, Category = Settings)
+	EOffsetRootBoneMode TranslationMode = EOffsetRootBoneMode::Interpolate;
+
+	// The rotation offset behavior mode
+	UPROPERTY(EditAnywhere, BlueprintReadWrite, Category = Settings)
+	EOffsetRootBoneMode RotationMode = EOffsetRootBoneMode::Interpolate;
+
+	// Controls how fast the translation offset is blended out
+	// Values closer to 0 make it faster
+	UPROPERTY(EditAnywhere, BlueprintReadWrite, Category = Settings)
+	float TranslationHalflife = 0.1f;
+
+	// Controls how fast the rotation offset is blended out
+	// Values closer to 0 make it faster
+	UPROPERTY(EditAnywhere, BlueprintReadWrite, Category = Settings)
+	float RotationHalfLife = 0.2f;
+
+	// How much the offset can deviate from the mesh component's translation in units
+	// Values lower than 0 disable this limit
+	UPROPERTY(EditAnywhere, BlueprintReadWrite, Category = Settings)
+	float MaxTranslationError = -1.0f;
+
+	// How much the offset can deviate from the mesh component's rotation in degrees
+	// Values lower than 0 disable this limit
+	UPROPERTY(EditAnywhere, BlueprintReadWrite, Category = Settings)
+	float MaxRotationError = -1.0f;
+
+	// Whether to limit the offset's translation interpolation speed to the velocity on the incoming motion
+	// Enabling this prevents the offset sliding when there's little to no translation speed
+	UPROPERTY(EditAnywhere, BlueprintReadWrite, Category = Settings)
+	bool bClampToTranslationVelocity = false;
+
+	// Whether to limit the offset's rotation interpolation speed to the velocity on the incoming motion
+	// Enabling this prevents the offset sliding when there's little to no rotation speed
+	UPROPERTY(EditAnywhere, BlueprintReadWrite, Category = Settings)
+	bool bClampToRotationVelocity = false;
+
+	// How much the offset can blend out, relative to the incoming translation speed
+	// i.e. If root motion is moving at 400cm/s, at 0.5, the offset can blend out at 200cm/s
+	UPROPERTY(EditAnywhere, BlueprintReadWrite, Category = Settings, meta = (EditCondition = bClampToTranslationVelocity))
+	float TranslationSpeedRatio = 0.5f;
+
+	// How much the offset can blend out, relative to the incoming rotation speed
+	// i.e. If root motion is rotating at 90 degrees/s, at 0.5, the offset can blend out at 45 degree/s
+	UPROPERTY(EditAnywhere, BlueprintReadWrite, Category = Settings, meta = (EditCondition = bClampToRotationVelocity))
+	float RotationSpeedRatio = 0.5f;
+};
+
+USTRUCT(BlueprintInternalUseOnly, Experimental)
+struct ANIMATIONWARPINGRUNTIME_API FAnimNode_OffsetRootBone : public FAnimNode_SkeletalControlBase
+{
+	GENERATED_BODY()
+
+#if WITH_EDITORONLY_DATA
+	UPROPERTY(EditAnywhere, Category = Evaluation, meta=(FoldProperty))
+	EWarpingEvaluationMode EvaluationMode = EWarpingEvaluationMode::Graph;
+
+	UPROPERTY(EditAnywhere, Category = Settings, meta = (FoldProperty, PinHiddenByDefault))
+	FOffsetRootBoneSettings OffsetSettings;
+
+	// Delta applied to the translation offset this frame. 
+	// For procedural values, consider adjusting the input by delta time.
+	UPROPERTY(EditAnywhere, Category = Evaluation, meta = (FoldProperty, PinHiddenByDefault))
+	FVector TranslationDelta;
+
+	// Delta applied to the rotation offset this frame. 
+	// For procedural values, consider adjusting the input by delta time.
+	UPROPERTY(EditAnywhere, Category = Evaluation, meta = (FoldProperty, PinHiddenByDefault))
+	FRotator RotationDelta;
+#endif
+
+public:
+	// FAnimNode_Base interface
+	virtual void GatherDebugData(FNodeDebugData& DebugData) override;
+	virtual void UpdateInternal(const FAnimationUpdateContext& Context) override;
+	// End of FAnimNode_Base interface
+
+	// FAnimNode_SkeletalControlBase interface
+	virtual void Initialize_AnyThread(const FAnimationInitializeContext& Context) override;
+	virtual void EvaluateSkeletalControl_AnyThread(FComponentSpacePoseContext& Output, TArray<FBoneTransform>& OutBoneTransforms) override;
+	virtual bool IsValidToEvaluate(const USkeleton* Skeleton, const FBoneContainer& RequiredBones) override;
+	// End of FAnimNode_SkeletalControlBase interface
+
+	EWarpingEvaluationMode GetEvaluationMode() const;
+	const FOffsetRootBoneSettings& GetOffsetSettings() const;
+	const FVector& GetTranslationDelta() const;
+	const FRotator& GetRotationDelta() const;
+
+private:
+
+	// Internal cached anim instance proxy
+	FAnimInstanceProxy* AnimInstanceProxy = nullptr;
+
+	// Internal cached delta time used for interpolators
+	float CachedDeltaTime = 0.f;
+
+	bool bIsFirstUpdate = true;
+
+	FTransform ComponentTransform = FTransform::Identity;
+
+	// The simulated world-space transforms for the root bone with offset
+	// Offset = ComponentTransform - SimulatedTransform
+	FVector SimulatedTranslation = FVector::ZeroVector;
+	FQuat SimulatedRotation = FQuat::Identity;
+};
