Offset Root Bone: Fix for CIS error
#rnx

diff --git a/Engine/Plugins/Animation/AnimationWarping/Source/Runtime/Public/BoneControllers/AnimNode_OffsetRootBone.h b/Engine/Plugins/Animation/AnimationWarping/Source/Runtime/Public/BoneControllers/AnimNode_OffsetRootBone.h
--- a/Engine/Plugins/Animation/AnimationWarping/Source/Runtime/Public/BoneControllers/AnimNode_OffsetRootBone.h
+++ b/Engine/Plugins/Animation/AnimationWarping/Source/Runtime/Public/BoneControllers/AnimNode_OffsetRootBone.h
@@ -96,12 +96,12 @@ struct ANIMATIONWARPINGRUNTIME_API FAnimNode_OffsetRootBone : public FAnimNode_S
 	// Delta applied to the translation offset this frame. 
 	// For procedural values, consider adjusting the input by delta time.
 	UPROPERTY(EditAnywhere, Category = Evaluation, meta = (FoldProperty, PinHiddenByDefault))
-	FVector TranslationDelta;
+	FVector TranslationDelta = FVector::ZeroVector;
 
 	// Delta applied to the rotation offset this frame. 
 	// For procedural values, consider adjusting the input by delta time.
 	UPROPERTY(EditAnywhere, Category = Evaluation, meta = (FoldProperty, PinHiddenByDefault))
-	FRotator RotationDelta;
+	FRotator RotationDelta = FRotator::ZeroRotator;
 #endif
 
 public:
