Foot placement: Removed unused variables
#rnx

diff --git a/Engine/Plugins/Animation/AnimationWarping/Source/Runtime/Public/BoneControllers/AnimNode_FootPlacement.h b/Engine/Plugins/Animation/AnimationWarping/Source/Runtime/Public/BoneControllers/AnimNode_FootPlacement.h
--- a/Engine/Plugins/Animation/AnimationWarping/Source/Runtime/Public/BoneControllers/AnimNode_FootPlacement.h
+++ b/Engine/Plugins/Animation/AnimationWarping/Source/Runtime/Public/BoneControllers/AnimNode_FootPlacement.h
@@ -353,13 +353,6 @@ public:
 
 	UPROPERTY(EditAnywhere, Category = "Plant Settings")
 	EFootPlacementLockType LockType = EFootPlacementLockType::PivotAroundBall;
-	// Whether we fully lock rotation too.
-	UPROPERTY(EditAnywhere, Category = "Plant Settings")
-	bool bLockRotation = false;
-
-	// Whether we look the position of the foot along the tip bone (true) or the ik bone (false)
-	UPROPERTY(EditAnywhere, Category = "Plant Settings")
-	bool bPivotAroundBall = true;
 
 	// How much angular deviation (in degrees) causes the constraint to be released for replant
 	UPROPERTY(EditAnywhere, Category = "Plant Settings")
