Foot Placement:
-Improve pelvis placement when legs are close to hyper-compressed in the animated pose i.e. crouching
-Fixed bug when IK Foot traces higher than the original pelvis position
-Re-initialize foot placement interpolation when becoming relevant

diff --git a/Engine/Plugins/Animation/AnimationWarping/Source/Runtime/Private/BoneControllers/AnimNode_FootPlacement.cpp b/Engine/Plugins/Animation/AnimationWarping/Source/Runtime/Private/BoneControllers/AnimNode_FootPlacement.cpp
--- a/Engine/Plugins/Animation/AnimationWarping/Source/Runtime/Private/BoneControllers/AnimNode_FootPlacement.cpp
+++ b/Engine/Plugins/Animation/AnimationWarping/Source/Runtime/Private/BoneControllers/AnimNode_FootPlacement.cpp
@@ -217,7 +217,7 @@ void FAnimNode_FootPlacement::FindPelvisOffsetRangeForLimb(
 	const FVector& PlantTargetLocationCS,
 	const FTransform& PelvisTransformCS,
 	const float LimbLength,
-	FFloatRange& OutPelvisOffsetRangeCS) const
+	FPelvisOffsetRangeForLimb& OutPelvisOffsetRangeCS) const
 {
 	// TODO: Cache this.
 	const FTransform HipToPelvis =
@@ -228,8 +228,6 @@ void FAnimNode_FootPlacement::FindPelvisOffsetRangeForLimb(
 	const FVector DesiredExtensionDelta =
 		LegInputPose.FKTransformCS.GetLocation() - LegInputPose.HipTransformCS.GetLocation();
 
-	OutPelvisOffsetRangeCS = FFloatRange::Inclusive(0.0f, 0.0f);
-
 	const float DesiredExtensionSqrd = DesiredExtensionDelta.SizeSquared();
 	const float DesiredExtension = FMath::Sqrt(DesiredExtensionSqrd);
 	const float MaxExtension = GetMaxLimbExtension(DesiredExtension, LimbLength);
@@ -240,7 +238,6 @@ void FAnimNode_FootPlacement::FindPelvisOffsetRangeForLimb(
 	FVector DesiredPlantTargetLocationCS = PlantTargetLocationCS;
 	FVector MaxPlantTargetLocationCS = PlantTargetLocationCS;
 	// If the foot wants to be place so high up relative to the FK hip, this is unlikely to matter.
-	// TODO: Crouch might be an exception. Revisit when looking into that.
 	if (HipToPlantDotApproachDir > 0.0f)
 	{
 		const float OpposingSideSqrd = HipToPlantCS.SizeSquared() - HipToPlantDotApproachDir * HipToPlantDotApproachDir;
@@ -269,7 +266,7 @@ void FAnimNode_FootPlacement::FindPelvisOffsetRangeForLimb(
 
 			// Find the max horizontal offset. 
 			// We don't care about the circle intersection in the opposite direction
-			const float MaxOpposingSide = FMath::Sqrt(-DesiredHeightSqrd + RadiusSqrd);
+			const float MaxOpposingSide = FMath::Sqrt(FMath::Abs(-DesiredHeightSqrd + RadiusSqrd));
 
 			// Respected the input pose if it exceeds it
 			const float MaxIKOrthogonalDist = FMath::Max(FKFootToHipDist, MaxOpposingSide);
@@ -297,10 +294,19 @@ void FAnimNode_FootPlacement::FindPelvisOffsetRangeForLimb(
 	FVector DesiredOffsetLocation;
 	FMath::SphereDistToLine(MaxPlantTargetLocationCS, MaxExtension, HipLocationCS - Context.ApproachDirCS * TraceSettings.EndOffset, Context.ApproachDirCS, MaxOffsetLocation);
 	FMath::SphereDistToLine(DesiredPlantTargetLocationCS, DesiredExtension, HipLocationCS - Context.ApproachDirCS * TraceSettings.EndOffset, Context.ApproachDirCS, DesiredOffsetLocation);
+	
 	const float MaxOffset = (MaxOffsetLocation - HipLocationCS) | -Context.ApproachDirCS;
 	const float DesiredOffset = (DesiredOffsetLocation - HipLocationCS) | -Context.ApproachDirCS;
+	OutPelvisOffsetRangeCS.MaxExtension = MaxOffset;
+	OutPelvisOffsetRangeCS.DesiredExtension = DesiredOffset;
 
-	OutPelvisOffsetRangeCS = FFloatRange::Inclusive(DesiredOffset, MaxOffset);
+	// Calculate min offset considering only the height of the foot
+	// Poses where the foot's height is close to the hip's height are bad. 
+	const float MinExtension = GetMinLimbExtension(DesiredExtension, LimbLength);
+	const FVector MinOffsetLocation = DesiredPlantTargetLocationCS + -Context.ApproachDirCS * MinExtension;
+
+	const float MinOffset = (MinOffsetLocation - HipLocationCS) | -Context.ApproachDirCS;
+	OutPelvisOffsetRangeCS.MinExtension = MinOffset;
 }
 
 float FAnimNode_FootPlacement::CalcTargetPlantPlaneDistance(
@@ -928,6 +934,15 @@ void FAnimNode_FootPlacement::Initialize_AnyThread(const FAnimationInitializeCon
 void FAnimNode_FootPlacement::UpdateInternal(const FAnimationUpdateContext& Context)
 {
 	FAnimNode_SkeletalControlBase::UpdateInternal(Context);
+
+	// If we just became relevant and haven't been initialized yet, then reinitialize foot placement.
+	if (!bIsFirstUpdate && UpdateCounter.HasEverBeenUpdated() && !UpdateCounter.WasSynchronizedCounter(Context.AnimInstanceProxy->GetUpdateCounter()))
+	{
+		FAnimationInitializeContext InitializationContext(Context.AnimInstanceProxy, Context.SharedContext);
+		Initialize_AnyThread(InitializationContext);
+	}
+	UpdateCounter.SynchronizeWith(Context.AnimInstanceProxy->GetUpdateCounter());
+
 	CachedDeltaTime += Context.GetDeltaTime();
 }
 
@@ -1178,8 +1193,8 @@ void FAnimNode_FootPlacement::GatherLegDataFromInputs(
 			LegData.InputPose.FKTransformCS * Context.OwningComponentToWorld;
 		LegData.UnalignedFootTransformWS = LegData.AlignedFootTransformWS;
 
-		const FVector IKFootRootLocationWS = 
-			Context.OwningComponentToWorld.InverseTransformPosition(PelvisData.InputPose.IKRootTransformCS.GetLocation());
+		const FVector IKFootRootLocationWS =
+			Context.OwningComponentToWorld.TransformPosition(PelvisData.InputPose.IKRootTransformCS.GetLocation());
 
 		LegData.Plant.PlantPlaneWS = FPlane(IKFootRootLocationWS, -Context.ApproachDirWS);
 		LegData.Plant.PlantPlaneCS = FPlane(PelvisData.InputPose.IKRootTransformCS.GetLocation(), -Context.ApproachDirCS);
@@ -1384,10 +1399,11 @@ FTransform FAnimNode_FootPlacement::SolvePelvis(const UE::Anim::FootPlacement::F
 	float MaxOffsetMin = BIG_NUMBER;
 	float DesiredOffsetMin = BIG_NUMBER;
 	float DesiredOffsetSum = 0.0f;
+	float MinOffsetMax = -BIG_NUMBER;
 
 	for (const FLegRuntimeData& LegData : LegsData)
 	{
-		FFloatRange PelvisOffsetRangeCS;
+		FPelvisOffsetRangeForLimb PelvisOffsetRangeCS;
 		FindPelvisOffsetRangeForLimb(
 			Context,
 			LegData.InputPose,
@@ -1396,12 +1412,14 @@ FTransform FAnimNode_FootPlacement::SolvePelvis(const UE::Anim::FootPlacement::F
 			LegData.Bones.LimbLength,
 			PelvisOffsetRangeCS);
 
-		const float DesiredOffset = PelvisOffsetRangeCS.GetLowerBoundValue();
-		const float MaxOffset = PelvisOffsetRangeCS.GetUpperBoundValue();
+		const float DesiredOffset = PelvisOffsetRangeCS.DesiredExtension;
+		const float MaxOffset = PelvisOffsetRangeCS.MaxExtension;
+		const float MinOffset = PelvisOffsetRangeCS.MinExtension;
 
 		DesiredOffsetSum += DesiredOffset;
 		DesiredOffsetMin = FMath::Min(DesiredOffsetMin, DesiredOffset);
 		MaxOffsetMin = FMath::Min(MaxOffsetMin, MaxOffset);
+		MinOffsetMax = FMath::Max(MinOffsetMax, MinOffset);
 	}
 	const float DesiredOffsetAvg = DesiredOffsetSum / LegsData.Num();
 	const float MinToAvg = DesiredOffsetAvg - DesiredOffsetMin;
@@ -1409,14 +1427,16 @@ FTransform FAnimNode_FootPlacement::SolvePelvis(const UE::Anim::FootPlacement::F
 
 	DesiredOffsetMin -= 0.05f;
 
-	// TODO: This doesn't deal well w/ over-compression. 
 	// In cases like crouching, it favors over-compressing to preserve the pose of the other leg
 	// Consider working in over-compression into the formula.
 	const float Divisor = MinToAvg + MinToMax;
-	const float PelvisOffsetZ = FMath::IsNearlyZero(Divisor) ?
+	float PelvisOffsetZ = FMath::IsNearlyZero(Divisor) ?
 		DesiredOffsetMin :
 		DesiredOffsetMin + ((MinToAvg * MinToMax) / Divisor);
 
+	// Adjust the hips to prevent over-compression
+	PelvisOffsetZ = FMath::Clamp(PelvisOffsetZ, MinOffsetMax, MaxOffsetMin);
+
 	const FVector PelvisOffsetDelta = -PelvisOffsetZ * Context.ApproachDirCS;
 	FTransform PelvisTransformCS = PelvisData.InputPose.FKTransformCS;
 	PelvisTransformCS.AddToTranslation(PelvisOffsetDelta);
diff --git a/Engine/Plugins/Animation/AnimationWarping/Source/Runtime/Public/BoneControllers/AnimNode_FootPlacement.h b/Engine/Plugins/Animation/AnimationWarping/Source/Runtime/Public/BoneControllers/AnimNode_FootPlacement.h
--- a/Engine/Plugins/Animation/AnimationWarping/Source/Runtime/Public/BoneControllers/AnimNode_FootPlacement.h
+++ b/Engine/Plugins/Animation/AnimationWarping/Source/Runtime/Public/BoneControllers/AnimNode_FootPlacement.h
@@ -487,16 +487,21 @@ private:
 		const UE::Anim::FootPlacement::FEvaluationContext& Context,
 		const UE::Anim::FootPlacement::FLegRuntimeData::FInputPoseData& LegInputPose) const;
 
+	struct FPelvisOffsetRangeForLimb
+	{
+		float MaxExtension;
+		float MinExtension;
+		float DesiredExtension;
+	};
+
 	// Find the horizontal pelvis offset range for the foot to reach:
-	// OutPelvisOffsetRangeCS.GetLowerBoundValue() = Offset for Desired extension
-	// OutPelvisOffsetRangeCS.GetUperBoundValue() = Offset for Max extension
 	void FindPelvisOffsetRangeForLimb(
 		const UE::Anim::FootPlacement::FEvaluationContext& Context,
 		const UE::Anim::FootPlacement::FLegRuntimeData::FInputPoseData& LegInputPoseData,
 		const FVector& PlantTargetLocationCS,
 		const FTransform& PelvisTransformCS,
 		const float LimbLength,
-		FFloatRange& OutPelvisOffsetRangeCS) const;
+		FPelvisOffsetRangeForLimb& OutPelvisOffsetRangeCS) const;
 
 	// Adjust LastPlantTransformWS to current, to have the foot pivot around the ball instead of the ankle
 	FTransform GetFootPivotAroundBallWS(
@@ -549,4 +554,5 @@ private:
 #endif
 
 	bool bIsFirstUpdate = false;
+	FGraphTraversalCounter UpdateCounter;
 };
