Added graph-based foot/ball speed calculation for Foot Placement

diff --git a/Engine/Plugins/Animation/AnimationWarping/Source/Runtime/Private/BoneControllers/AnimNode_FootPlacement.cpp b/Engine/Plugins/Animation/AnimationWarping/Source/Runtime/Private/BoneControllers/AnimNode_FootPlacement.cpp
--- a/Engine/Plugins/Animation/AnimationWarping/Source/Runtime/Private/BoneControllers/AnimNode_FootPlacement.cpp
+++ b/Engine/Plugins/Animation/AnimationWarping/Source/Runtime/Private/BoneControllers/AnimNode_FootPlacement.cpp
@@ -6,6 +6,7 @@
 #include "GameFramework/Character.h"
 #include "GameFramework/CharacterMovementComponent.h"
 #include "Components/CapsuleComponent.h"
+#include "Animation/AnimRootMotionProvider.h"
 
 DECLARE_CYCLE_STAT(TEXT("Foot Placement Eval"), STAT_FootPlacement_Eval, STATGROUP_Anim);
 
@@ -35,6 +36,7 @@ namespace UE::Anim::FootPlacement
 		class UWorld* World = nullptr;
 		class UCharacterMovementComponent* MovementComponent = nullptr;
 		FTransform OwningComponentToWorld = FTransform::Identity;
+		FTransform RootMotionTransformDelta = FTransform::Identity;
 		float UpdateDeltaTime = 0.0f;
 		FVector ApproachDirWS = -FVector::UpVector;
 		FVector ApproachDirCS = -FVector::UpVector;
@@ -78,6 +80,12 @@ namespace UE::Anim::FootPlacement
 		OwningComponentToWorld = OwningComponent->GetComponentToWorld();
 
 		ApproachDirWS = OwningComponentToWorld.TransformVector(ApproachDirCS);
+	
+		RootMotionTransformDelta = FTransform::Identity;
+		if (const UE::Anim::IAnimRootMotionProvider* RootMotionProvider = UE::Anim::IAnimRootMotionProvider::Get())
+		{
+			RootMotionProvider->ExtractRootMotion(CSPContext.CustomAttributes, RootMotionTransformDelta);
+		}
 	}
 
 	static FVector ReOrientNormal(const FVector& ApproachDir, const FVector& InNormal, FVector& PointA, const FVector& PointB)
@@ -1173,6 +1181,8 @@ void FAnimNode_FootPlacement::GatherLegDataFromInputs(
 	UE::Anim::FootPlacement::FLegRuntimeData& LegData,
 	const FFootPlacemenLegDefinition& LegDef)
 {
+	FVector LastBallLocation = LegData.InputPose.BallTransformCS.GetLocation();
+
 	LegData.InputPose.FKTransformCS =
 		Context.CSPContext.Pose.GetComponentSpaceTransform(LegData.Bones.FKIndex);
 	LegData.InputPose.IKTransformCS =
@@ -1201,13 +1211,30 @@ void FAnimNode_FootPlacement::GatherLegDataFromInputs(
 
 		LegData.Plant.PlantType = UE::Anim::FootPlacement::EPlantType::Unplanted;
 		LegData.Plant.LastPlantType = UE::Anim::FootPlacement::EPlantType::Unplanted;
+		LastBallLocation = LegData.InputPose.BallTransformCS.GetLocation();
+	}
+
+	if (PlantSpeedMode == EWarpingEvaluationMode::Graph)
+	{
+		FVector BallTranslationDelta = LegData.InputPose.BallTransformCS.GetLocation() - LastBallLocation;
+
+		// Apply root motion delta to the ball's translation delta in root space
+		const FQuat RootRotation = Context.CSPContext.Pose.GetComponentSpaceTransform(FCompactPoseBoneIndex(0)).GetRotation();
+		const FVector CorrectedRootMotionTranslationDelta = RootRotation.RotateVector(Context.RootMotionTransformDelta.GetTranslation());
+		BallTranslationDelta += CorrectedRootMotionTranslationDelta;
+
+		const float BallDeltaDistance = BallTranslationDelta.Size();
+		LegData.InputPose.Speed = BallDeltaDistance / Context.UpdateDeltaTime;
+	}
+	else
+	{
+		bool bValidSpeedCurve;
+		// If the curve is not found in the stream, assume we're unplanted.
+		const float DefaultSpeedCurveValue = PlantSettings.SpeedThreshold;
+		LegData.InputPose.Speed =
+			Context.CSPContext.Curve.Get(LegData.SpeedCurveUID, bValidSpeedCurve, DefaultSpeedCurveValue);
 	}
 
-	bool bValidSpeedCurve;
-	// If the curve is not found in the stream, assume we're unplanted.
-	const float DefaultSpeedCurveValue = PlantSettings.SpeedThreshold;
-	LegData.InputPose.Speed =
-		Context.CSPContext.Curve.Get(LegData.SpeedCurveUID, bValidSpeedCurve, DefaultSpeedCurveValue);
 	LegData.InputPose.DistanceToPlant = CalcTargetPlantPlaneDistance(Context, LegData.InputPose);
 	const float FKAlignmentAlpha = GetAlignmentAlpha(Context, LegData.InputPose);
 	LegData.InputPose.AlignmentAlpha = FKAlignmentAlpha;
diff --git a/Engine/Plugins/Animation/AnimationWarping/Source/Runtime/Public/BoneControllers/AnimNode_FootPlacement.h b/Engine/Plugins/Animation/AnimationWarping/Source/Runtime/Public/BoneControllers/AnimNode_FootPlacement.h
--- a/Engine/Plugins/Animation/AnimationWarping/Source/Runtime/Public/BoneControllers/AnimNode_FootPlacement.h
+++ b/Engine/Plugins/Animation/AnimationWarping/Source/Runtime/Public/BoneControllers/AnimNode_FootPlacement.h
@@ -312,6 +312,7 @@ public:
 	UPROPERTY(EditAnywhere, Category = "Settings")
 	int32 NumBonesInLimb = 2;
 
+	// Name of the curve representing the foot/ball speed. Not required in Graph speed mode
 	UPROPERTY(EditAnywhere, Category = "Settings")
 	FName SpeedCurveName = NAME_None;
 
@@ -390,6 +391,13 @@ struct ANIMATIONWARPINGRUNTIME_API FAnimNode_FootPlacement : public FAnimNode_Sk
 	GENERATED_BODY()
 
 public:
+
+	// Foot/Ball speed evaluation mode (Graph or Manual) used to decide when the feet are locked
+	// Graph mode uses the root motion attribute from the animations to calculate the joint's speed
+	// Manual mode uses a per-foot curve name representing the joint's speed
+	UPROPERTY(EditAnywhere, Category = "Settings")
+	EWarpingEvaluationMode PlantSpeedMode = EWarpingEvaluationMode::Manual;
+
 	UPROPERTY(EditAnywhere, Category = "Settings")
 	TArray<FFootPlacemenLegDefinition> LegDefinitions;
 
