Foot Placement:
-Improve behavior when replant limits are the same as unplant limits
#rnx

diff --git a/Engine/Plugins/Animation/AnimationWarping/Source/Runtime/Private/BoneControllers/AnimNode_FootPlacement.cpp b/Engine/Plugins/Animation/AnimationWarping/Source/Runtime/Private/BoneControllers/AnimNode_FootPlacement.cpp
--- a/Engine/Plugins/Animation/AnimationWarping/Source/Runtime/Private/BoneControllers/AnimNode_FootPlacement.cpp
+++ b/Engine/Plugins/Animation/AnimationWarping/Source/Runtime/Private/BoneControllers/AnimNode_FootPlacement.cpp
@@ -544,9 +544,12 @@ void FAnimNode_FootPlacement::DeterminePlantType(
 			CurrentBoneTransformWS.GetLocation() - FKTransformWS.GetLocation();
 
 		// TODO: Test along approach direction
+		// Don't consider the limits to be exceeded if replant radius == unplant radius.
 		const bool bPlantTranslationExceeded =
+			PlantSettings.ReplantMaxLinearErrorRatio < 1.0f &&
 			PlantTranslationWS.SizeSquared2D() > PlantRuntimeSettings.MaxLinearErrorSqrd;
 		const bool bPlantRotationExceeded =
+			PlantSettings.ReplantMaxRotationErrorRatio < 1.0f &&
 			FMath::Abs(InOutPlantData.TwistCorrection.W) <
 			PlantRuntimeSettings.CosHalfMaxRotationError;
 
@@ -875,12 +878,15 @@ void FAnimNode_FootPlacement::DrawDebug(
 		PlantCenter, UnplantRadius, 24, PlantedColor,
 		LegData.Plant.PlantPlaneWS.GetNormal(), false, -1.0f, SDPG_Foreground, 0.5f);
 
-	const float ReplantRadius =
-		PlantSettings.MaxLinearError *
-		PlantSettings.ReplantMaxLinearErrorRatio;
-	Context.CSPContext.AnimInstanceProxy->AnimDrawDebugCircle(
-		PlantCenter, ReplantRadius, 24, ReplantedColor,
-		LegData.Plant.PlantPlaneWS.GetNormal(), false, -1.0f, SDPG_Foreground, 0.5f);
+	if (PlantSettings.ReplantMaxLinearErrorRatio < 1.0f)
+	{
+		const float ReplantRadius =
+			PlantSettings.MaxLinearError *
+			PlantSettings.ReplantMaxLinearErrorRatio;
+		Context.CSPContext.AnimInstanceProxy->AnimDrawDebugCircle(
+			PlantCenter, ReplantRadius, 24, ReplantedColor,
+			LegData.Plant.PlantPlaneWS.GetNormal(), false, -1.0f, SDPG_Foreground, 0.5f);
+	}
 
 	
 	FString InputPoseMessage = FString::Format(
@@ -1376,6 +1382,32 @@ void FAnimNode_FootPlacement::ProcessFootAlignment(
 		}
 	}
 
+	// If replant radius is the same as unplant radius, clamp the location and slide
+	if (PlantSettings.ReplantMaxLinearErrorRatio >= 1.0f)
+	{
+		const FVector ClampedTransltionOffset = Interpolation.UnalignedFootOffsetCS.GetLocation().GetClampedToMaxSize(PlantSettings.MaxLinearError);
+		Interpolation.UnalignedFootOffsetCS.SetLocation(ClampedTransltionOffset);
+	}
+
+	// If replant angle is the same as unplant angle, clamp the angle and slide
+	if (PlantSettings.ReplantMaxRotationErrorRatio >= 1.0f)
+	{
+		FQuat ClampedRotationOffset = Interpolation.UnalignedFootOffsetCS.GetRotation();
+		ClampedRotationOffset.Normalize();
+		ClampedRotationOffset = ClampedRotationOffset.W < 0.0 ? -ClampedRotationOffset : ClampedRotationOffset;
+
+		FVector OffsetAxis;
+		float OffsetAngle;
+		ClampedRotationOffset.ToAxisAndAngle(OffsetAxis, OffsetAngle);
+
+		const float MaxAngle = FMath::DegreesToRadians(PlantSettings.MaxRotationError);
+		if (FMath::Abs(OffsetAngle) > MaxAngle)
+		{
+			ClampedRotationOffset = FQuat(OffsetAxis, MaxAngle);
+		}
+		Interpolation.UnalignedFootOffsetCS.SetRotation(ClampedRotationOffset);
+	}
+
 	const FTransform IKUnalignedTransformCS = InputPose.IKTransformCS * Interpolation.UnalignedFootOffsetCS;
 	LegData.UnalignedFootTransformWS = IKUnalignedTransformCS * Context.OwningComponentToWorld;
 
diff --git a/Engine/Plugins/Animation/AnimationWarping/Source/Runtime/Public/BoneControllers/AnimNode_FootPlacement.h b/Engine/Plugins/Animation/AnimationWarping/Source/Runtime/Public/BoneControllers/AnimNode_FootPlacement.h
--- a/Engine/Plugins/Animation/AnimationWarping/Source/Runtime/Public/BoneControllers/AnimNode_FootPlacement.h
+++ b/Engine/Plugins/Animation/AnimationWarping/Source/Runtime/Public/BoneControllers/AnimNode_FootPlacement.h
@@ -337,11 +337,11 @@ public:
 	float SpeedThreshold = 60.0f;
 
 	// Max extension ratio of the chain, calculated from the remaining length between current pose and full limb extension
-	UPROPERTY(EditAnywhere, Category = "Plant Settings")
+	UPROPERTY(EditAnywhere, Category = "Plant Settings", meta = (ClampMin = "0.0", UIMin = "0.0", ClampMax = "1.0", UIMax = "1.0"))
 	float MaxExtensionRatio = 0.5f;
 
 	// Min extension ratio of the chain, calculated from the total limb length, and adjusted along the approach direction
-	UPROPERTY(EditAnywhere, Category = "Plant Settings")
+	UPROPERTY(EditAnywhere, Category = "Plant Settings", meta = (ClampMin = "0.0", UIMin = "0.0", ClampMax = "1.0", UIMax = "1.0"))
 	float MinExtensionRatio = 0.2f;
 
 	// How much linear deviation causes the constraint to be released
@@ -349,7 +349,7 @@ public:
 	float MaxLinearError = 35.0f;
 
 	// Below this value, proportional to MaxLinearError, the bone will replant
-	UPROPERTY(EditAnywhere, Category = "Plant Settings")
+	UPROPERTY(EditAnywhere, Category = "Plant Settings", meta = (ClampMin = "0.0", UIMin = "0.0", ClampMax = "1.0", UIMax = "1.0"))
 	float ReplantMaxLinearErrorRatio = 0.35f;
 
 	UPROPERTY(EditAnywhere, Category = "Plant Settings")
@@ -360,7 +360,7 @@ public:
 	float MaxRotationError = 45.0;
 
 	// Below this value, proportional to MaxRotationError, the bone will replant
-	UPROPERTY(EditAnywhere, Category = "Plant Settings")
+	UPROPERTY(EditAnywhere, Category = "Plant Settings", meta = (ClampMin = "0.0", UIMin = "0.0", ClampMax = "1.0", UIMax = "1.0"))
 	float ReplantMaxRotationErrorRatio = 0.5f;
 
 	// Speed at which we transition to fully unplanted.
@@ -369,9 +369,7 @@ public:
 	UPROPERTY(EditAnywhere, Category = "Plant Settings")
 	float UnalignmentSpeedThreshold = 200.0f;
 
-	// Speed at which we transition to fully unplanted.
-	// The range between SpeedThreshold and UnalignmentSpeedThreshold should roughly represent the roll-phase of the foot
-	// TODO: This feels innaccurate most of the time, and varies depending on anim speed. Improve this
+	// How much we reduce the procedural ankle twist adjustment used to align the foot to the ground slope.
 	UPROPERTY(EditAnywhere, Category = "Plant Settings")
 	float AnkleTwistReduction = 0.75f;
 
