Stride Warping:
-Improve stride warping behavior on low speeds
-Fix on manual mode when stride scale is 1 and interpolation is on
-Fix for 3x interpolation speed in manual mode

diff --git a/Engine/Plugins/Animation/AnimationWarping/Source/Editor/Private/AnimGraph/AnimGraphNode_StrideWarping.cpp b/Engine/Plugins/Animation/AnimationWarping/Source/Editor/Private/AnimGraph/AnimGraphNode_StrideWarping.cpp
--- a/Engine/Plugins/Animation/AnimationWarping/Source/Editor/Private/AnimGraph/AnimGraphNode_StrideWarping.cpp
+++ b/Engine/Plugins/Animation/AnimationWarping/Source/Editor/Private/AnimGraph/AnimGraphNode_StrideWarping.cpp
@@ -53,7 +53,7 @@ void UAnimGraphNode_StrideWarping::CustomizePinData(UEdGraphPin* Pin, FName Sour
 		Pin->bHidden = (Node.Mode == EWarpingEvaluationMode::Manual);
 	}
 
-	if (Pin->PinName == GET_MEMBER_NAME_STRING_CHECKED(FAnimNode_StrideWarping, MinLocomotionSpeedThreshold))
+	if (Pin->PinName == GET_MEMBER_NAME_STRING_CHECKED(FAnimNode_StrideWarping, MinRootMotionSpeedThreshold))
 	{
 		Pin->bHidden = (Node.Mode == EWarpingEvaluationMode::Manual);
 	}
@@ -115,7 +115,7 @@ void UAnimGraphNode_StrideWarping::CustomizeDetails(IDetailLayoutBuilder& Detail
 	if (Node.Mode == EWarpingEvaluationMode::Manual)
 	{
 		DetailBuilder.HideProperty(NodeHandle->GetChildHandle(GET_MEMBER_NAME_CHECKED(FAnimNode_StrideWarping, LocomotionSpeed)));
-		DetailBuilder.HideProperty(NodeHandle->GetChildHandle(GET_MEMBER_NAME_CHECKED(FAnimNode_StrideWarping, MinLocomotionSpeedThreshold)));
+		DetailBuilder.HideProperty(NodeHandle->GetChildHandle(GET_MEMBER_NAME_CHECKED(FAnimNode_StrideWarping, MinRootMotionSpeedThreshold)));
 	}
 }
 
@@ -175,7 +175,7 @@ void UAnimGraphNode_StrideWarping::PostEditChangeProperty(struct FPropertyChange
 						Pin->BreakAllPinLinks();
 					}
 				}
-				else if (Pin->PinName == GET_MEMBER_NAME_STRING_CHECKED(FAnimNode_StrideWarping, MinLocomotionSpeedThreshold))
+				else if (Pin->PinName == GET_MEMBER_NAME_STRING_CHECKED(FAnimNode_StrideWarping, MinRootMotionSpeedThreshold))
 				{
 					if (Node.Mode == EWarpingEvaluationMode::Manual)
 					{
diff --git a/Engine/Plugins/Animation/AnimationWarping/Source/Runtime/Private/BoneControllers/AnimNode_StrideWarping.cpp b/Engine/Plugins/Animation/AnimationWarping/Source/Runtime/Private/BoneControllers/AnimNode_StrideWarping.cpp
--- a/Engine/Plugins/Animation/AnimationWarping/Source/Runtime/Private/BoneControllers/AnimNode_StrideWarping.cpp
+++ b/Engine/Plugins/Animation/AnimationWarping/Source/Runtime/Private/BoneControllers/AnimNode_StrideWarping.cpp
@@ -29,7 +29,7 @@ void FAnimNode_StrideWarping::GatherDebugData(FNodeDebugData& DebugData)
 			DebugLine += FString::Printf(TEXT("\n - Root Motion Direction: (%s)"), *(CachedRootMotionDeltaTranslation.GetSafeNormal().ToCompactString()));
 			DebugLine += FString::Printf(TEXT("\n - Root Motion Speed: (%.3fd)"), CachedRootMotionDeltaSpeed);
 #endif
-			DebugLine += FString::Printf(TEXT("\n - Min Locomotion Speed Threshold: (%.3fd)"), MinLocomotionSpeedThreshold);
+			DebugLine += FString::Printf(TEXT("\n - Min Locomotion Speed Threshold: (%.3fd)"), MinRootMotionSpeedThreshold);
 		}
 		DebugLine += FString::Printf(TEXT("\n - Floor Normal: (%s)"), *(FloorNormalDirection.Value.ToCompactString()));
 		DebugLine += FString::Printf(TEXT("\n - Gravity Direction: (%s)"), *(GravityDirection.Value.ToCompactString()));
@@ -68,6 +68,7 @@ void FAnimNode_StrideWarping::EvaluateSkeletalControl_AnyThread(FComponentSpaceP
 	SCOPE_CYCLE_COUNTER(STAT_StrideWarping_Eval);
 	check(OutBoneTransforms.IsEmpty());
 
+	const FVector PreviousStrideDirection = ActualStrideDirection;
 	ActualStrideDirection = StrideDirection;
 	ActualStrideScale = StrideScale;
 
@@ -93,7 +94,8 @@ void FAnimNode_StrideWarping::EvaluateSkeletalControl_AnyThread(FComponentSpaceP
 		if (bGraphDrivenWarping)
 		{
 			CachedRootMotionDeltaTranslation = RootMotionTransformDelta.GetTranslation();
-			ActualStrideDirection = CachedRootMotionDeltaTranslation.GetSafeNormal();
+			// If there's no root motion delta, keep the previous stride direction
+			ActualStrideDirection = CachedRootMotionDeltaTranslation.GetSafeNormal(UE_SMALL_NUMBER, PreviousStrideDirection);
 #if WITH_EDITORONLY_DATA
 			// Graph driven Stride Warping expects a root motion delta to be present in the attribute stream.
 			bFoundRootMotionAttribute = true;
@@ -159,10 +161,19 @@ void FAnimNode_StrideWarping::EvaluateSkeletalControl_AnyThread(FComponentSpaceP
 #else
 		const float CachedRootMotionDeltaSpeed = CachedRootMotionDeltaTranslation.Size() / CachedDeltaTime;
 #endif
-		// Graph driven stride scale factor will be determined by the ratio of the
-		// locomotion (capsule/physics) speed against the animation root motion speed
-		// If there's (almost) no extracted root motion, we don't want impossibly large strides.
-		ActualStrideScale = FMath::IsNearlyEqual(CachedRootMotionDeltaSpeed, 0.f, KINDA_SMALL_NUMBER) ? 1.0f : LocomotionSpeed / CachedRootMotionDeltaSpeed;
+
+		if (CachedRootMotionDeltaSpeed <= MinRootMotionSpeedThreshold)
+		{
+			// If root motion speed is under the threshold, snap back to no stride adjustment.
+			// If interpolation is on, it will blend the result
+			ActualStrideScale = 1.0f;
+		}
+		else
+		{
+			// Graph driven stride scale factor will be determined by the ratio of the
+			// locomotion (capsule/physics) speed against the animation root motion speed
+			ActualStrideScale = LocomotionSpeed / CachedRootMotionDeltaSpeed;
+		}
 	}
 
 	// Allow the opportunity for stride scale clamping and interpolation regardless of evaluation mode
@@ -372,18 +383,6 @@ bool FAnimNode_StrideWarping::IsValidToEvaluate(const USkeleton* Skeleton, const
 		}
 	}
 
-	if (Mode == EWarpingEvaluationMode::Manual)
-	{
-		if (FMath::IsNearlyEqual(StrideScaleModifierState.ApplyTo(StrideScaleModifier, StrideScale, CachedDeltaTime), 1.f, KINDA_SMALL_NUMBER))
-		{
-			return false;
-		}
-	}
-	else if (LocomotionSpeed <= MinLocomotionSpeedThreshold)
-	{
-		return false;
-	}
-
 	return true;
 }
 
diff --git a/Engine/Plugins/Animation/AnimationWarping/Source/Runtime/Public/BoneControllers/AnimNode_StrideWarping.h b/Engine/Plugins/Animation/AnimationWarping/Source/Runtime/Public/BoneControllers/AnimNode_StrideWarping.h
--- a/Engine/Plugins/Animation/AnimationWarping/Source/Runtime/Public/BoneControllers/AnimNode_StrideWarping.h
+++ b/Engine/Plugins/Animation/AnimationWarping/Source/Runtime/Public/BoneControllers/AnimNode_StrideWarping.h
@@ -54,9 +54,11 @@ struct ANIMATIONWARPINGRUNTIME_API FAnimNode_StrideWarping : public FAnimNode_Sk
 	UPROPERTY(EditAnywhere, BlueprintReadWrite, Category=Evaluation, meta=(ClampMin="0.0", PinShownByDefault))
 	float LocomotionSpeed = 0.f;
 
-	// Minimum locomotion speed required to apply stride warping
-	UPROPERTY(EditAnywhere, BlueprintReadWrite, Category=Evaluation, meta=(ClampMin="0.0", PinHiddenByDefault))
-	float MinLocomotionSpeedThreshold = 10.f;
+	// Minimum root motion speed required to apply stride warping
+	// This is useful to prevent unnatural strides when the animation has a portion with no root motion (i.e starts/stops)
+	// When this value is greater than 0, it's recommended to enable interpolation in StrideScaleModifier
+	UPROPERTY(EditAnywhere, BlueprintReadWrite, Category=Evaluation, meta=(ClampMin="0.0001", PinHiddenByDefault))
+	float MinRootMotionSpeedThreshold = 10.0f;
 
 	// Pevlis Bone definition
 	UPROPERTY(EditAnywhere, Category=Settings)
