Naming and ordering pass on Foot Placement
#rnx

diff --git a/Engine/Plugins/Animation/AnimationWarping/Source/Runtime/Private/BoneControllers/AnimNode_FootPlacement.cpp b/Engine/Plugins/Animation/AnimationWarping/Source/Runtime/Private/BoneControllers/AnimNode_FootPlacement.cpp
--- a/Engine/Plugins/Animation/AnimationWarping/Source/Runtime/Private/BoneControllers/AnimNode_FootPlacement.cpp
+++ b/Engine/Plugins/Animation/AnimationWarping/Source/Runtime/Private/BoneControllers/AnimNode_FootPlacement.cpp
@@ -546,12 +546,12 @@ void FAnimNode_FootPlacement::DeterminePlantType(
 		// TODO: Test along approach direction
 		// Don't consider the limits to be exceeded if replant radius == unplant radius.
 		const bool bPlantTranslationExceeded =
-			PlantSettings.ReplantMaxLinearErrorRatio < 1.0f &&
-			PlantTranslationWS.SizeSquared2D() > PlantRuntimeSettings.MaxLinearErrorSqrd;
+			PlantSettings.ReplantRadiusRatio < 1.0f &&
+			PlantTranslationWS.SizeSquared2D() > PlantRuntimeSettings.UnplantRadiusSqrd;
 		const bool bPlantRotationExceeded =
-			PlantSettings.ReplantMaxRotationErrorRatio < 1.0f &&
+			PlantSettings.ReplantAngleRatio < 1.0f &&
 			FMath::Abs(InOutPlantData.TwistCorrection.W) <
-			PlantRuntimeSettings.CosHalfMaxRotationError;
+			PlantRuntimeSettings.CosHalfUnplantAngle;
 
 		if (!bPlantTranslationExceeded && !bPlantRotationExceeded)
 		{
@@ -573,10 +573,10 @@ void FAnimNode_FootPlacement::DeterminePlantType(
 		const float LocationDeltaSizeSqrd = PlantLocationDelta.SizeSquared2D();
 
 		const bool bLocationWithinBounds =
-			LocationDeltaSizeSqrd <= PlantRuntimeSettings.ReplantMaxLinearErrorSqrd;
+			LocationDeltaSizeSqrd <= PlantRuntimeSettings.ReplantRadiusSqrd;
 		const bool bTwistWithinBounds =
 			FMath::Abs(InOutPlantData.TwistCorrection.W) >=
-			PlantRuntimeSettings.CosHalfReplantMaxRotationError;
+			PlantRuntimeSettings.CosHalfReplantAngle;
 
 		if (bLocationWithinBounds && bTwistWithinBounds)
 		{
@@ -617,7 +617,7 @@ bool FAnimNode_FootPlacement::WantsToPlant(
 		return false;
 	}
 
-	const bool bPassesPlantDistanceCheck = LegInputPose.DistanceToPlant < PlantSettings.FKPlantDistance;
+	const bool bPassesPlantDistanceCheck = LegInputPose.DistanceToPlant < PlantSettings.DistanceToGround;
 	const bool bPassesSpeedCheck = LegInputPose.Speed < PlantSettings.SpeedThreshold;
 	return bPassesPlantDistanceCheck && bPassesSpeedCheck;
 }
@@ -869,7 +869,7 @@ void FAnimNode_FootPlacement::DrawDebug(
 		CurrentPlantColor,
 		false, -1.0f, 1.0f, SDPG_Foreground);
 
-	const float UnplantRadius = PlantSettings.MaxLinearError;
+	const float UnplantRadius = PlantSettings.UnplantRadius;
 	const FVector PlantCenter = FMath::LinePlaneIntersection(
 		IKBoneTransformWS.GetLocation(),
 		IKBoneTransformWS.GetLocation() + Context.ApproachDirWS,
@@ -878,30 +878,29 @@ void FAnimNode_FootPlacement::DrawDebug(
 		PlantCenter, UnplantRadius, 24, PlantedColor,
 		LegData.Plant.PlantPlaneWS.GetNormal(), false, -1.0f, SDPG_Foreground, 0.5f);
 
-	if (PlantSettings.ReplantMaxLinearErrorRatio < 1.0f)
+	if (PlantSettings.ReplantRadiusRatio < 1.0f)
 	{
 		const float ReplantRadius =
-			PlantSettings.MaxLinearError *
-			PlantSettings.ReplantMaxLinearErrorRatio;
+			PlantSettings.UnplantRadius *
+			PlantSettings.ReplantRadiusRatio;
 		Context.CSPContext.AnimInstanceProxy->AnimDrawDebugCircle(
 			PlantCenter, ReplantRadius, 24, ReplantedColor,
 			LegData.Plant.PlantPlaneWS.GetNormal(), false, -1.0f, SDPG_Foreground, 0.5f);
 	}
 
-	
-	FString InputPoseMessage = FString::Format(
-		TEXT("{0}\n\t - InputPose [ AlignmentAlpha = {1}, Speed = {2}, DistanceToPlant = {3}]"), 
-		{	*LegDefinitions[LegData.Idx].FKFootBone.BoneName.ToString(),
+	FString InputPoseMessage = FString::Printf(
+		TEXT("%s\n\t - InputPose [ AlignmentAlpha = %.2f, Speed = %.2f, DistanceToPlant = %.2f]"), 
+			*LegDefinitions[LegData.Idx].FKFootBone.BoneName.ToString(),
 			LegData.InputPose.AlignmentAlpha,
 			LegData.InputPose.Speed,
-			LegData.InputPose.DistanceToPlant });
+			LegData.InputPose.DistanceToPlant );
 	Context.CSPContext.AnimInstanceProxy->AnimDrawDebugOnScreenMessage(InputPoseMessage, FColor::White);
 	
-	FString ExtensionMessage = FString::Format(
-		TEXT("\t - HyperExtension[ Amount = {1}, Roll = {2}, Pull {3}]"),
-		{	 DebugData.LegsExtension[LegData.Idx].HyperExtensionAmount,
+	FString ExtensionMessage = FString::Printf(
+		TEXT("\t - HyperExtension[ Amount = %.2f, Roll = %.2f, Pull %.2f]"),
+			DebugData.LegsExtension[LegData.Idx].HyperExtensionAmount,
 			DebugData.LegsExtension[LegData.Idx].RollAmount,
-			DebugData.LegsExtension[LegData.Idx].PullAmount });
+			DebugData.LegsExtension[LegData.Idx].PullAmount);
 	Context.CSPContext.AnimInstanceProxy->AnimDrawDebugOnScreenMessage(ExtensionMessage,
 		(DebugData.LegsExtension[LegData.Idx].HyperExtensionAmount <= 0.0f) ? FColor::Green : FColor::Red);
 
@@ -988,14 +987,14 @@ void FAnimNode_FootPlacement::EvaluateSkeletalControl_AnyThread(FComponentSpaceP
 		// will keep this for now. If needed, change to lazy update.
 		PlantRuntimeSettings.MaxExtensionRatioSqrd = PlantSettings.MaxExtensionRatio * PlantSettings.MaxExtensionRatio;
 		PlantRuntimeSettings.MinExtensionRatioSqrd = PlantSettings.MinExtensionRatio * PlantSettings.MinExtensionRatio;
-		PlantRuntimeSettings.MaxLinearErrorSqrd = PlantSettings.MaxLinearError * PlantSettings.MaxLinearError;
-		PlantRuntimeSettings.ReplantMaxLinearErrorSqrd =
-			PlantRuntimeSettings.MaxLinearErrorSqrd *
-			PlantSettings.ReplantMaxLinearErrorRatio * PlantSettings.ReplantMaxLinearErrorRatio;
-		PlantRuntimeSettings.CosHalfMaxRotationError = FMath::Cos(FMath::DegreesToRadians(PlantSettings.MaxRotationError / 2.0f));
-		PlantRuntimeSettings.CosHalfReplantMaxRotationError
+		PlantRuntimeSettings.UnplantRadiusSqrd = PlantSettings.UnplantRadius * PlantSettings.UnplantRadius;
+		PlantRuntimeSettings.ReplantRadiusSqrd =
+			PlantRuntimeSettings.UnplantRadiusSqrd *
+			PlantSettings.ReplantRadiusRatio * PlantSettings.ReplantRadiusRatio;
+		PlantRuntimeSettings.CosHalfUnplantAngle = FMath::Cos(FMath::DegreesToRadians(PlantSettings.UnplantAngle / 2.0f));
+		PlantRuntimeSettings.CosHalfReplantAngle
 			= FMath::Cos(FMath::DegreesToRadians(
-				(PlantSettings.MaxRotationError * PlantSettings.ReplantMaxRotationErrorRatio) / 2.0f));
+				(PlantSettings.UnplantAngle * PlantSettings.ReplantAngleRatio) / 2.0f));
 	}
 
 	ProcessCharacterState(FootPlacementContext);
@@ -1010,6 +1009,10 @@ void FAnimNode_FootPlacement::EvaluateSkeletalControl_AnyThread(FComponentSpaceP
 	PelvisTransformCS = UpdatePelvisInterpolation(FootPlacementContext, PelvisTransformCS);
 	OutBoneTransforms.Add(FBoneTransform(PelvisData.Bones.FkBoneIndex, PelvisTransformCS));
 
+#if ENABLE_ANIM_DEBUG
+	FString HeaderMessage = FString::Printf(TEXT("FOOT PLACEMENT DEBUG"));
+	FootPlacementContext.CSPContext.AnimInstanceProxy->AnimDrawDebugOnScreenMessage(HeaderMessage, FColor::Cyan);
+#endif
 
 	for (int32 FootIndex = 0; FootIndex < LegsData.Num(); ++FootIndex)
 	{
@@ -1383,14 +1386,14 @@ void FAnimNode_FootPlacement::ProcessFootAlignment(
 	}
 
 	// If replant radius is the same as unplant radius, clamp the location and slide
-	if (PlantSettings.ReplantMaxLinearErrorRatio >= 1.0f)
+	if (PlantSettings.ReplantRadiusRatio >= 1.0f)
 	{
-		const FVector ClampedTransltionOffset = Interpolation.UnalignedFootOffsetCS.GetLocation().GetClampedToMaxSize(PlantSettings.MaxLinearError);
+		const FVector ClampedTransltionOffset = Interpolation.UnalignedFootOffsetCS.GetLocation().GetClampedToMaxSize(PlantSettings.UnplantRadius);
 		Interpolation.UnalignedFootOffsetCS.SetLocation(ClampedTransltionOffset);
 	}
 
 	// If replant angle is the same as unplant angle, clamp the angle and slide
-	if (PlantSettings.ReplantMaxRotationErrorRatio >= 1.0f)
+	if (PlantSettings.ReplantAngleRatio >= 1.0f)
 	{
 		FQuat ClampedRotationOffset = Interpolation.UnalignedFootOffsetCS.GetRotation();
 		ClampedRotationOffset.Normalize();
@@ -1400,7 +1403,7 @@ void FAnimNode_FootPlacement::ProcessFootAlignment(
 		float OffsetAngle;
 		ClampedRotationOffset.ToAxisAndAngle(OffsetAxis, OffsetAngle);
 
-		const float MaxAngle = FMath::DegreesToRadians(PlantSettings.MaxRotationError);
+		const float MaxAngle = FMath::DegreesToRadians(PlantSettings.UnplantAngle);
 		if (FMath::Abs(OffsetAngle) > MaxAngle)
 		{
 			ClampedRotationOffset = FQuat(OffsetAxis, MaxAngle);
diff --git a/Engine/Plugins/Animation/AnimationWarping/Source/Runtime/Public/BoneControllers/AnimNode_FootPlacement.h b/Engine/Plugins/Animation/AnimationWarping/Source/Runtime/Public/BoneControllers/AnimNode_FootPlacement.h
--- a/Engine/Plugins/Animation/AnimationWarping/Source/Runtime/Public/BoneControllers/AnimNode_FootPlacement.h
+++ b/Engine/Plugins/Animation/AnimationWarping/Source/Runtime/Public/BoneControllers/AnimNode_FootPlacement.h
@@ -88,10 +88,10 @@ namespace UE::Anim::FootPlacement
 	{
 		float MaxExtensionRatioSqrd = 0.0f;
 		float MinExtensionRatioSqrd = 0.0f;
-		float MaxLinearErrorSqrd = 0.0f;
-		float ReplantMaxLinearErrorSqrd = 0.0f;
-		float CosHalfMaxRotationError = 0.0f;
-		float CosHalfReplantMaxRotationError = 0.0f;
+		float UnplantRadiusSqrd = 0.0f;
+		float ReplantRadiusSqrd = 0.0f;
+		float CosHalfUnplantAngle = 0.0f;
+		float CosHalfReplantAngle = 0.0f;
 	};
 
 	struct FPelvisRuntimeData
@@ -328,40 +328,40 @@ struct FFootPlacementPlantSettings
 
 public:
 	
-	// At this distance from the planting plane the bone is considered planted and will be fully aligned.
-	UPROPERTY(EditAnywhere, Category = "Plant Settings")
-	float FKPlantDistance = 10.0f;
-
 	// Bone is considered planted below this speed. Value is obtained from FKSpeedCurveName
 	UPROPERTY(EditAnywhere, Category = "Plant Settings")
 	float SpeedThreshold = 60.0f;
 
-	// Max extension ratio of the chain, calculated from the remaining length between current pose and full limb extension
-	UPROPERTY(EditAnywhere, Category = "Plant Settings", meta = (ClampMin = "0.0", UIMin = "0.0", ClampMax = "1.0", UIMax = "1.0"))
-	float MaxExtensionRatio = 0.5f;
+	// At this distance from the planting plane the bone is considered planted and will be fully aligned.
+	UPROPERTY(EditAnywhere, Category = "Plant Settings")
+	float DistanceToGround = 10.0f;
 
-	// Min extension ratio of the chain, calculated from the total limb length, and adjusted along the approach direction
-	UPROPERTY(EditAnywhere, Category = "Plant Settings", meta = (ClampMin = "0.0", UIMin = "0.0", ClampMax = "1.0", UIMax = "1.0"))
-	float MinExtensionRatio = 0.2f;
+	UPROPERTY(EditAnywhere, Category = "Plant Settings")
+	EFootPlacementLockType LockType = EFootPlacementLockType::PivotAroundBall;
 
 	// How much linear deviation causes the constraint to be released
 	UPROPERTY(EditAnywhere, Category = "Plant Settings")
-	float MaxLinearError = 35.0f;
+	float UnplantRadius = 35.0f;
 
-	// Below this value, proportional to MaxLinearError, the bone will replant
+	// Below this value, proportional to UnplantRadius, the bone will replant
 	UPROPERTY(EditAnywhere, Category = "Plant Settings", meta = (ClampMin = "0.0", UIMin = "0.0", ClampMax = "1.0", UIMax = "1.0"))
-	float ReplantMaxLinearErrorRatio = 0.35f;
-
-	UPROPERTY(EditAnywhere, Category = "Plant Settings")
-	EFootPlacementLockType LockType = EFootPlacementLockType::PivotAroundBall;
+	float ReplantRadiusRatio = 0.35f;
 
 	// How much angular deviation (in degrees) causes the constraint to be released for replant
 	UPROPERTY(EditAnywhere, Category = "Plant Settings")
-	float MaxRotationError = 45.0;
+	float UnplantAngle = 45.0;
+
+	// Below this value, proportional to UnplantAngle, the bone will replant
+	UPROPERTY(EditAnywhere, Category = "Plant Settings", meta = (ClampMin = "0.0", UIMin = "0.0", ClampMax = "1.0", UIMax = "1.0"))
+	float ReplantAngleRatio = 0.5f;
+
+	// Max extension ratio of the chain, calculated from the remaining length between current pose and full limb extension
+	UPROPERTY(EditAnywhere, Category = "Plant Settings", meta = (ClampMin = "0.0", UIMin = "0.0", ClampMax = "1.0", UIMax = "1.0"))
+	float MaxExtensionRatio = 0.5f;
 
-	// Below this value, proportional to MaxRotationError, the bone will replant
+	// Min extension ratio of the chain, calculated from the total limb length, and adjusted along the approach direction
 	UPROPERTY(EditAnywhere, Category = "Plant Settings", meta = (ClampMin = "0.0", UIMin = "0.0", ClampMax = "1.0", UIMax = "1.0"))
-	float ReplantMaxRotationErrorRatio = 0.5f;
+	float MinExtensionRatio = 0.2f;
 
 	// Speed at which we transition to fully unplanted.
 	// The range between SpeedThreshold and UnalignmentSpeedThreshold should roughly represent the roll-phase of the foot
@@ -370,7 +370,7 @@ public:
 	float UnalignmentSpeedThreshold = 200.0f;
 
 	// How much we reduce the procedural ankle twist adjustment used to align the foot to the ground slope.
-	UPROPERTY(EditAnywhere, Category = "Plant Settings")
+	UPROPERTY(EditAnywhere, Category = "Plant Settings", meta = (ClampMin = "0.0", UIMin = "0.0", ClampMax = "1.0", UIMax = "1.0"))
 	float AnkleTwistReduction = 0.75f;
 
 	// How much we can pull the the foot towards the hip to prevent hyperextension
@@ -396,28 +396,28 @@ public:
 	UPROPERTY(EditAnywhere, Category = "Settings")
 	EWarpingEvaluationMode PlantSpeedMode = EWarpingEvaluationMode::Manual;
 
+	// TODO: This wont work well when the animation doesn't have a single plant plane, i.e. a walking upstairs anim
+	UPROPERTY(EditAnywhere, Category = "Settings")
+	FBoneReference IKFootRootBone;
+
+	UPROPERTY(EditAnywhere, Category = "Settings")
+	FBoneReference PelvisBone;
+
+	UPROPERTY(EditAnywhere, BlueprintReadWrite, Category = "Settings", meta = (PinHiddenByDefault))
+	FFootPlacementPelvisSettings PelvisSettings;
+
 	UPROPERTY(EditAnywhere, Category = "Settings")
 	TArray<FFootPlacemenLegDefinition> LegDefinitions;
 
-	UPROPERTY(EditAnywhere, BlueprintReadWrite, Category = "Plant", meta = (PinHiddenByDefault))
+	UPROPERTY(EditAnywhere, BlueprintReadWrite, Category = "Settings", meta = (PinHiddenByDefault))
 	FFootPlacementPlantSettings PlantSettings;
 
-	UPROPERTY(EditAnywhere, BlueprintReadWrite, Category = "Interpolation", meta = (PinHiddenByDefault))
+	UPROPERTY(EditAnywhere, BlueprintReadWrite, Category = "Settings", meta = (PinHiddenByDefault))
 	FFootPlacementInterpolationSettings InterpolationSettings;
 
-	UPROPERTY(EditAnywhere, BlueprintReadWrite, Category = "Trace", meta = (PinHiddenByDefault))
+	UPROPERTY(EditAnywhere, BlueprintReadWrite, Category = "Settings", meta = (PinHiddenByDefault))
 	FFootPlacementTraceSettings TraceSettings;
 
-	// TODO: This wont work well when the animation doesn't have a single plant plane, i.e. a walking upstairs anim
-	UPROPERTY(EditAnywhere, Category = "Plant")
-	FBoneReference IKFootRootBone;
-
-	UPROPERTY(EditAnywhere, Category = "Pelvis")
-	FBoneReference PelvisBone;
-
-	UPROPERTY(EditAnywhere, BlueprintReadWrite, Category = "Pelvis", meta = (PinHiddenByDefault))
-	FFootPlacementPelvisSettings PelvisSettings;
-
 public:
 	FAnimNode_FootPlacement();
 
