Foot Placement:
-Fix for debug text showing up all the time
#rnx

diff --git a/Engine/Plugins/Animation/AnimationWarping/Source/Runtime/Private/BoneControllers/AnimNode_FootPlacement.cpp b/Engine/Plugins/Animation/AnimationWarping/Source/Runtime/Private/BoneControllers/AnimNode_FootPlacement.cpp
--- a/Engine/Plugins/Animation/AnimationWarping/Source/Runtime/Private/BoneControllers/AnimNode_FootPlacement.cpp
+++ b/Engine/Plugins/Animation/AnimationWarping/Source/Runtime/Private/BoneControllers/AnimNode_FootPlacement.cpp
@@ -1010,8 +1010,11 @@ void FAnimNode_FootPlacement::EvaluateSkeletalControl_AnyThread(FComponentSpaceP
 	OutBoneTransforms.Add(FBoneTransform(PelvisData.Bones.FkBoneIndex, PelvisTransformCS));
 
 #if ENABLE_ANIM_DEBUG
-	FString HeaderMessage = FString::Printf(TEXT("FOOT PLACEMENT DEBUG"));
-	FootPlacementContext.CSPContext.AnimInstanceProxy->AnimDrawDebugOnScreenMessage(HeaderMessage, FColor::Cyan);
+	if (CVarAnimNodeFootPlacementDebug.GetValueOnAnyThread())
+	{
+		FString HeaderMessage = FString::Printf(TEXT("FOOT PLACEMENT DEBUG"));
+		FootPlacementContext.CSPContext.AnimInstanceProxy->AnimDrawDebugOnScreenMessage(HeaderMessage, FColor::Cyan);
+	}
 #endif
 
 	for (int32 FootIndex = 0; FootIndex < LegsData.Num(); ++FootIndex)
