Foot Placement:
-Fix for bad behavior when trace hit fails
#rnx

diff --git a/Engine/Plugins/Animation/AnimationWarping/Source/Runtime/Private/BoneControllers/AnimNode_FootPlacement.cpp b/Engine/Plugins/Animation/AnimationWarping/Source/Runtime/Private/BoneControllers/AnimNode_FootPlacement.cpp
--- a/Engine/Plugins/Animation/AnimationWarping/Source/Runtime/Private/BoneControllers/AnimNode_FootPlacement.cpp
+++ b/Engine/Plugins/Animation/AnimationWarping/Source/Runtime/Private/BoneControllers/AnimNode_FootPlacement.cpp
@@ -136,6 +136,14 @@ namespace UE::Anim::FootPlacement
 		const bool bHit = Context.World->SweepSingleByChannel(
 			HitResult, TraceStart, TraceEnd, FQuat::Identity, CollisionChannel, CollisionShape, QueryParams);
 
+		if (!bHit)
+		{
+			// If the hit fails, use starting position and a default impact normal (negated trace direction)
+			OutImpactLocationWS = StartPositionWS;
+			OutImpactNormalWS = -TraceDirectionWS;
+			return false;
+		}
+
 		OutImpactLocationWS = HitResult.ImpactPoint;
 		if (!Context.GetMovementComponentIsWalkable(HitResult))
 		{
