Move AnimationWarping out of beta

diff --git a/Engine/Plugins/Animation/AnimationWarping/AnimationWarping.uplugin b/Engine/Plugins/Animation/AnimationWarping/AnimationWarping.uplugin
--- a/Engine/Plugins/Animation/AnimationWarping/AnimationWarping.uplugin
+++ b/Engine/Plugins/Animation/AnimationWarping/AnimationWarping.uplugin
@@ -9,7 +9,7 @@
 	"CreatedByURL": "http://epicgames.com",
 	"EnabledByDefault": false,
 	"CanContainContent": true,
-	"IsBetaVersion": true,
+	"IsBetaVersion": false,
 	"IsExperimentalVersion": false,
 	"Installed": false,
 	"Modules": [
