Foot Placement:
-Added separating plane to prevent feet from crossing
-Tweaked default pelvis drop due to foot locking

diff --git a/Engine/Plugins/Animation/AnimationWarping/Source/Runtime/Private/BoneControllers/AnimNode_FootPlacement.cpp b/Engine/Plugins/Animation/AnimationWarping/Source/Runtime/Private/BoneControllers/AnimNode_FootPlacement.cpp
--- a/Engine/Plugins/Animation/AnimationWarping/Source/Runtime/Private/BoneControllers/AnimNode_FootPlacement.cpp
+++ b/Engine/Plugins/Animation/AnimationWarping/Source/Runtime/Private/BoneControllers/AnimNode_FootPlacement.cpp
@@ -756,9 +756,9 @@ UE::Anim::FootPlacement::FPlantResult FAnimNode_FootPlacement::FinalizeFootAlign
 			}
 
 #if ENABLE_ANIM_DEBUG
-			DebugData.LegsExtension[LegData.Idx].HyperExtensionAmount = HyperExtensionAmount;
-			DebugData.LegsExtension[LegData.Idx].RollAmount = HyperExtensionAmount - HyperExtensionRemaining;
-			DebugData.LegsExtension[LegData.Idx].PullAmount = FMath::Max(0.0f, HyperExtensionRemaining);
+			DebugData.LegsInfo[LegData.Idx].HyperExtensionAmount = HyperExtensionAmount;
+			DebugData.LegsInfo[LegData.Idx].RollAmount = HyperExtensionAmount - HyperExtensionRemaining;
+			DebugData.LegsInfo[LegData.Idx].PullAmount = FMath::Max(0.0f, HyperExtensionRemaining);
 #endif // (ENABLE_ANIM_DEBUG)
 		}
 	}
@@ -906,15 +906,26 @@ void FAnimNode_FootPlacement::DrawDebug(
 	
 	FString ExtensionMessage = FString::Printf(
 		TEXT("\t - HyperExtension[ Amount = %.2f, Roll = %.2f, Pull %.2f]"),
-			DebugData.LegsExtension[LegData.Idx].HyperExtensionAmount,
-			DebugData.LegsExtension[LegData.Idx].RollAmount,
-			DebugData.LegsExtension[LegData.Idx].PullAmount);
+			DebugData.LegsInfo[LegData.Idx].HyperExtensionAmount,
+			DebugData.LegsInfo[LegData.Idx].RollAmount,
+			DebugData.LegsInfo[LegData.Idx].PullAmount);
 	Context.CSPContext.AnimInstanceProxy->AnimDrawDebugOnScreenMessage(ExtensionMessage,
-		(DebugData.LegsExtension[LegData.Idx].HyperExtensionAmount <= 0.0f) ? FColor::Green : FColor::Red);
+		(DebugData.LegsInfo[LegData.Idx].HyperExtensionAmount <= 0.0f) ? FColor::Green : FColor::Red);
 
-	TRACE_ANIM_NODE_VALUE(Context.CSPContext, TEXT("HyperExtension - Amount"), DebugData.LegsExtension[LegData.Idx].HyperExtensionAmount);
-	TRACE_ANIM_NODE_VALUE(Context.CSPContext, TEXT("HyperExtension - Roll"), DebugData.LegsExtension[LegData.Idx].RollAmount);
-	TRACE_ANIM_NODE_VALUE(Context.CSPContext, TEXT("HyperExtension - Pull"), DebugData.LegsExtension[LegData.Idx].PullAmount);
+	if (PlantSettings.SeparatingDistance > 0.0f)
+	{
+		FString SeparationPlaneMessage = FString::Printf(
+			TEXT("\t - Distance To Separating Plane = %.2f"),
+			DebugData.LegsInfo[LegData.Idx].DistanceToSeparatingPlane);
+		Context.CSPContext.AnimInstanceProxy->AnimDrawDebugOnScreenMessage(SeparationPlaneMessage,
+			(DebugData.LegsInfo[LegData.Idx].DistanceToSeparatingPlane < 0.0f) ? FColor::Red : FColor::Green);
+
+		TRACE_ANIM_NODE_VALUE(Context.CSPContext, TEXT("DistanceToSeparatingPlane"), DebugData.LegsInfo[LegData.Idx].DistanceToSeparatingPlane);
+	}
+
+	TRACE_ANIM_NODE_VALUE(Context.CSPContext, TEXT("HyperExtension - Amount"), DebugData.LegsInfo[LegData.Idx].HyperExtensionAmount);
+	TRACE_ANIM_NODE_VALUE(Context.CSPContext, TEXT("HyperExtension - Roll"), DebugData.LegsInfo[LegData.Idx].RollAmount);
+	TRACE_ANIM_NODE_VALUE(Context.CSPContext, TEXT("HyperExtension - Pull"), DebugData.LegsInfo[LegData.Idx].PullAmount);
 	TRACE_ANIM_NODE_VALUE(Context.CSPContext, TEXT("InputPose - AlignmentAlpha"), LegData.InputPose.AlignmentAlpha);
 }
 #endif
@@ -1006,7 +1017,7 @@ void FAnimNode_FootPlacement::EvaluateSkeletalControl_AnyThread(FComponentSpaceP
 	}
 
 	ProcessCharacterState(FootPlacementContext);
-
+	CalculateFootMidpoint(FootPlacementContext, LegsData, PelvisData.InputPose.FootMidpointCS);
 	for (UE::Anim::FootPlacement::FLegRuntimeData& LegData : LegsData)
 	{
 		ProcessFootAlignment(FootPlacementContext, LegData);
@@ -1260,6 +1271,16 @@ void FAnimNode_FootPlacement::GatherLegDataFromInputs(
 	LegData.InputPose.AlignmentAlpha = FKAlignmentAlpha;
 }
 
+void FAnimNode_FootPlacement::CalculateFootMidpoint(const UE::Anim::FootPlacement::FEvaluationContext& Context, TConstArrayView<UE::Anim::FootPlacement::FLegRuntimeData> InLegsData, FVector& OutMidpoint) const
+{
+	int32 NumLegs = LegsData.Num();
+	OutMidpoint = FVector::ZeroVector;
+	for (const UE::Anim::FootPlacement::FLegRuntimeData& LegData : InLegsData)
+	{
+		OutMidpoint += LegData.InputPose.IKTransformCS.GetLocation() / NumLegs;
+	}
+}
+
 void FAnimNode_FootPlacement::ProcessCharacterState(const UE::Anim::FootPlacement::FEvaluationContext& Context)
 {
 	const FVector LastComponentLocationWS = bIsFirstUpdate
@@ -1422,7 +1443,39 @@ void FAnimNode_FootPlacement::ProcessFootAlignment(
 		Interpolation.UnalignedFootOffsetCS.SetRotation(ClampedRotationOffset);
 	}
 
-	const FTransform IKUnalignedTransformCS = InputPose.IKTransformCS * Interpolation.UnalignedFootOffsetCS;
+	FTransform IKUnalignedTransformCS = InputPose.IKTransformCS * Interpolation.UnalignedFootOffsetCS;
+	if (PlantSettings.SeparatingDistance > 0.0f)
+	{
+		// Prevent the feet from crossing by enforcing a set distance from a plane at the midpoint between all feet
+		FVector IKUnalignedLocationCS = IKUnalignedTransformCS.GetLocation();
+		const FVector MidPointToFoot = (InputPose.IKTransformCS.GetLocation() - PelvisData.InputPose.FootMidpointCS);
+		const FVector PlaneNormal = MidPointToFoot.GetSafeNormal2D();
+		const FVector PlaneCenter = PelvisData.InputPose.FootMidpointCS + PlaneNormal * PlantSettings.SeparatingDistance;
+		const FPlane SeparatingPlane = FPlane(PlaneCenter, PlaneNormal);
+		
+		const float DistanceToSeparatingPlane =
+			UE::Anim::FootPlacement::GetDistanceToPlaneAlongDirection(IKUnalignedLocationCS, SeparatingPlane, -PlaneNormal);
+
+		if (DistanceToSeparatingPlane < 0.0f)
+		{
+			IKUnalignedLocationCS -= PlaneNormal * DistanceToSeparatingPlane;
+			IKUnalignedTransformCS.SetLocation(IKUnalignedLocationCS);
+		}
+
+#if ENABLE_ANIM_DEBUG
+		if (CVarAnimNodeFootPlacementDebug.GetValueOnAnyThread())
+		{
+			const FVector PlaneNormalWS = Context.OwningComponentToWorld.TransformVector(PlaneNormal);;
+			const FVector PlaneCenterWS = Context.OwningComponentToWorld.TransformPosition(PlaneCenter);;
+			Context.CSPContext.AnimInstanceProxy->AnimDrawDebugCircle(
+				PlaneCenterWS, 25.0f, 24, FColor::Red,
+				PlaneNormalWS, false, -1.0f, SDPG_Foreground, 0.5f);
+
+			DebugData.LegsInfo[LegData.Idx].DistanceToSeparatingPlane = DistanceToSeparatingPlane;
+		}
+#endif
+	}
+
 	LegData.UnalignedFootTransformWS = IKUnalignedTransformCS * Context.OwningComponentToWorld;
 
 	const FTransform ComponentToWorldInv = Context.OwningComponentToWorld.Inverse();
diff --git a/Engine/Plugins/Animation/AnimationWarping/Source/Runtime/Public/BoneControllers/AnimNode_FootPlacement.h b/Engine/Plugins/Animation/AnimationWarping/Source/Runtime/Public/BoneControllers/AnimNode_FootPlacement.h
--- a/Engine/Plugins/Animation/AnimationWarping/Source/Runtime/Public/BoneControllers/AnimNode_FootPlacement.h
+++ b/Engine/Plugins/Animation/AnimationWarping/Source/Runtime/Public/BoneControllers/AnimNode_FootPlacement.h
@@ -112,6 +112,7 @@ namespace UE::Anim::FootPlacement
 		{
 			FTransform FKTransformCS = FTransform::Identity;
 			FTransform IKRootTransformCS = FTransform::Identity;
+			FVector FootMidpointCS = FVector::ZeroVector;
 		} InputPose;
 
 		/* Interpolation */
@@ -146,19 +147,20 @@ namespace UE::Anim::FootPlacement
 		TArray<FVector> OutputFootLocationsWS;
 		TArray<FVector> InputFootLocationsWS;
 
-		struct FLegExtension
+		struct FLegInfo
 		{
 			float HyperExtensionAmount;
 			float RollAmount;
 			float PullAmount;
+			float DistanceToSeparatingPlane;
 		};
-		TArray<FLegExtension> LegsExtension;
+		TArray<FLegInfo> LegsInfo;
 
 		void Init(const int32 InSize)
 		{
 			OutputFootLocationsWS.SetNumUninitialized(InSize);
 			InputFootLocationsWS.SetNumUninitialized(InSize);
-			LegsExtension.SetNumUninitialized(InSize);
+			LegsInfo.SetNumUninitialized(InSize);
 		}
 	};
 #endif
@@ -268,12 +270,6 @@ struct FFootPlacementPelvisSettings
 	GENERATED_BODY()
 
 public:
-	// Max horizontal foot adjustment we consider to lower the hips
-	// This can be used to prevent the hips from dropping too low when the feet are locked
-	// Exceeding this value will first attempt to roll the planted feet, and then slide
-	UPROPERTY(EditAnywhere, Category = "Pelvis Settings")
-	float MaxOffsetHorizontal = 20.0f;
-
 	// Max vertical offset from the input pose for the Pelvis.
 	// Reaching this limit means the feet may not reach their plant plane
 	UPROPERTY(EditAnywhere, Category = "Pelvis Settings")
@@ -285,6 +281,12 @@ public:
 	UPROPERTY(EditAnywhere, Category = "Pelvis Settings")
 	float LinearDamping = 1.0f;
 
+	// Max horizontal foot adjustment we consider to lower the hips
+	// This can be used to prevent the hips from dropping too low when the feet are locked
+	// Exceeding this value will first attempt to roll the planted feet, and then slide
+	UPROPERTY(EditAnywhere, Category = "Pelvis Settings")
+	float MaxOffsetHorizontal = 10.0f;
+
 	// This is used to hold the Pelvis's interpolator in a fixed spot when the capsule suddenly moves (i.e. on a big step)
 	// If your camera is directly attached to the character with little to no smoothing, you may want this disabled
 	UPROPERTY(EditAnywhere, Category = "Pelvis Settings")
@@ -363,6 +365,11 @@ public:
 	UPROPERTY(EditAnywhere, Category = "Plant Settings", meta = (ClampMin = "0.0", UIMin = "0.0", ClampMax = "1.0", UIMax = "1.0"))
 	float MinExtensionRatio = 0.2f;
 
+	// The minimum distance the feet can be from the plane that separates the feet. 
+	// Value of 0 disables this
+	UPROPERTY(EditAnywhere, Category = "Plant Settings", meta = (ClampMin = "0.0", UIMin = "0.0"))
+	float SeparatingDistance = 0.0f;
+
 	// Speed at which we transition to fully unplanted.
 	// The range between SpeedThreshold and UnalignmentSpeedThreshold should roughly represent the roll-phase of the foot
 	// TODO: This feels innaccurate most of the time, and varies depending on anim speed. Improve this
@@ -444,6 +451,10 @@ private:
 		const UE::Anim::FootPlacement::FEvaluationContext& Context,
 		UE::Anim::FootPlacement::FLegRuntimeData& LegData,
 		const FFootPlacemenLegDefinition& LegDef);
+	void CalculateFootMidpoint(
+		const UE::Anim::FootPlacement::FEvaluationContext& Context,
+		TConstArrayView<UE::Anim::FootPlacement::FLegRuntimeData> LegData,
+		FVector& OutMidpoint) const;
 	void ProcessCharacterState(const UE::Anim::FootPlacement::FEvaluationContext& Context);
 	void ProcessFootAlignment(
 		const UE::Anim::FootPlacement::FEvaluationContext& Context,
