Foot Placement:
-Add pelvis rebalancing
-Fix showdebug animation on Foot Placement

diff --git a/Engine/Plugins/Animation/AnimationWarping/Source/Runtime/Private/BoneControllers/AnimNode_FootPlacement.cpp b/Engine/Plugins/Animation/AnimationWarping/Source/Runtime/Private/BoneControllers/AnimNode_FootPlacement.cpp
--- a/Engine/Plugins/Animation/AnimationWarping/Source/Runtime/Private/BoneControllers/AnimNode_FootPlacement.cpp
+++ b/Engine/Plugins/Animation/AnimationWarping/Source/Runtime/Private/BoneControllers/AnimNode_FootPlacement.cpp
@@ -937,7 +937,7 @@ FAnimNode_FootPlacement::FAnimNode_FootPlacement()
 // TODO: implement 
 void FAnimNode_FootPlacement::GatherDebugData(FNodeDebugData& NodeDebugData)
 {
-
+	ComponentPose.GatherDebugData(NodeDebugData);
 }
 
 void FAnimNode_FootPlacement::Initialize_AnyThread(const FAnimationInitializeContext& Context)
@@ -1563,7 +1563,24 @@ FTransform FAnimNode_FootPlacement::SolvePelvis(const UE::Anim::FootPlacement::F
 	// Adjust the hips to prevent over-compression
 	PelvisOffsetZ = FMath::Clamp(PelvisOffsetZ, MinOffsetMax, MaxOffsetMin);
 
-	const FVector PelvisOffsetDelta = -PelvisOffsetZ * Context.ApproachDirCS;
+	FVector PelvisOffsetDelta = -PelvisOffsetZ * Context.ApproachDirCS;
+
+	if (PelvisSettings.HorizontalRebalancingWeight)
+	{
+		const int32 NumLegs = LegsData.Num();
+		FVector OffsetAverage = FVector::ZeroVector;
+		for (const FLegRuntimeData& LegData : LegsData)
+		{
+			const FVector LegTranslationOffset = LegData.AlignedFootTransformCS.GetLocation() - LegData.InputPose.IKTransformCS.GetLocation();
+			OffsetAverage += LegTranslationOffset / NumLegs;
+		}
+
+		// Remove the vertical component
+		OffsetAverage -= OffsetAverage * Context.ApproachDirCS;
+
+		PelvisOffsetDelta += OffsetAverage * PelvisSettings.HorizontalRebalancingWeight;
+	}
+
 	FTransform PelvisTransformCS = PelvisData.InputPose.FKTransformCS;
 	PelvisTransformCS.AddToTranslation(PelvisOffsetDelta);
 
diff --git a/Engine/Plugins/Animation/AnimationWarping/Source/Runtime/Public/BoneControllers/AnimNode_FootPlacement.h b/Engine/Plugins/Animation/AnimationWarping/Source/Runtime/Public/BoneControllers/AnimNode_FootPlacement.h
--- a/Engine/Plugins/Animation/AnimationWarping/Source/Runtime/Public/BoneControllers/AnimNode_FootPlacement.h
+++ b/Engine/Plugins/Animation/AnimationWarping/Source/Runtime/Public/BoneControllers/AnimNode_FootPlacement.h
@@ -281,6 +281,12 @@ public:
 	UPROPERTY(EditAnywhere, Category = "Pelvis Settings")
 	float LinearDamping = 1.0f;
 
+	// How much we move the pelvis horizontally to re-balance the characters weight due to foot offsets.
+	// A value of 0 will disable this effect.
+	// Higher values may move the mesh outside of the character's capsule
+	UPROPERTY(EditAnywhere, Category = "Pelvis Settings", meta = (ClampMin = "0.0", UIMin = "0.0", ClampMax = "1.0", UIMax = "1.0"))
+	float HorizontalRebalancingWeight = 0.3f;
+
 	// Max horizontal foot adjustment we consider to lower the hips
 	// This can be used to prevent the hips from dropping too low when the feet are locked
 	// Exceeding this value will first attempt to roll the planted feet, and then slide
