Pose Warping:
-Allow disabling orientation warping when root motion information isn't enough to calculate a good warp
-Fix for sudden flips on orientation warping changes, when there wasn't previously any root motion information

diff --git a/Engine/Plugins/Animation/AnimationWarping/Source/Editor/Private/AnimGraph/AnimGraphNode_OrientationWarping.cpp b/Engine/Plugins/Animation/AnimationWarping/Source/Editor/Private/AnimGraph/AnimGraphNode_OrientationWarping.cpp
--- a/Engine/Plugins/Animation/AnimationWarping/Source/Editor/Private/AnimGraph/AnimGraphNode_OrientationWarping.cpp
+++ b/Engine/Plugins/Animation/AnimationWarping/Source/Editor/Private/AnimGraph/AnimGraphNode_OrientationWarping.cpp
@@ -106,6 +106,7 @@ void UAnimGraphNode_OrientationWarping::CustomizeDetails(IDetailLayoutBuilder& D
 	{
 		DetailBuilder.HideProperty(NodeHandle->GetChildHandle(GET_MEMBER_NAME_CHECKED(FAnimNode_OrientationWarping, LocomotionAngle)));
 		DetailBuilder.HideProperty(NodeHandle->GetChildHandle(GET_MEMBER_NAME_CHECKED(FAnimNode_OrientationWarping, LocomotionAngleDeltaThreshold)));
+		DetailBuilder.HideProperty(NodeHandle->GetChildHandle(GET_MEMBER_NAME_CHECKED(FAnimNode_OrientationWarping, MinRootMotionSpeedThreshold)));
 	}
 }
 
@@ -150,6 +151,13 @@ void UAnimGraphNode_OrientationWarping::PostEditChangeProperty(struct FPropertyC
 						Pin->BreakAllPinLinks();
 					}
 				}
+				else if (Pin->PinName == GET_MEMBER_NAME_STRING_CHECKED(FAnimNode_OrientationWarping, MinRootMotionSpeedThreshold))
+				{
+					if (Node.Mode == EWarpingEvaluationMode::Manual)
+					{
+						Pin->BreakAllPinLinks();
+					}
+				}
 			}
 
 			bRequiresNodeReconstruct = true;
diff --git a/Engine/Plugins/Animation/AnimationWarping/Source/Runtime/Private/BoneControllers/AnimNode_OrientationWarping.cpp b/Engine/Plugins/Animation/AnimationWarping/Source/Runtime/Private/BoneControllers/AnimNode_OrientationWarping.cpp
--- a/Engine/Plugins/Animation/AnimationWarping/Source/Runtime/Private/BoneControllers/AnimNode_OrientationWarping.cpp
+++ b/Engine/Plugins/Animation/AnimationWarping/Source/Runtime/Private/BoneControllers/AnimNode_OrientationWarping.cpp
@@ -92,6 +92,7 @@ void FAnimNode_OrientationWarping::EvaluateSkeletalControl_AnyThread(FComponentS
 
 	ActualOrientationAngle = OrientationAngle;
 
+	const float DeltaSeconds = Output.AnimInstanceProxy->GetDeltaSeconds();
 	const FVector RotationAxisVector = UE::Anim::GetAxisVector(RotationAxis);
 	FVector RootMotionDeltaDirection = FVector::ZeroVector;
 	FVector LocomotionForward = FVector::ZeroVector;
@@ -154,32 +155,45 @@ void FAnimNode_OrientationWarping::EvaluateSkeletalControl_AnyThread(FComponentS
 			LocomotionForward = SkeletalMeshRelativeRotation.UnrotateVector(LocomotionRotation.GetForwardVector()).GetSafeNormal();
 
 			const FVector RootMotionDeltaTranslation = RootMotionTransformDelta.GetTranslation();
-			RootMotionDeltaDirection = RootMotionDeltaTranslation.GetSafeNormal();
+			// Hold previous direction if we can't calculate it from current move delta
+			RootMotionDeltaDirection = RootMotionDeltaTranslation.GetSafeNormal(UE_SMALL_NUMBER, PreviousRootMotionDeltaDirection);
 
-			// Capture the delta rotation from the axis of motion we care about
-			FQuat WarpedRotation = FQuat::FindBetween(RootMotionDeltaDirection, LocomotionForward);
-			ActualOrientationAngle = WarpedRotation.GetTwistAngle(RotationAxisVector);
-
-			// Motion Matching may return an animation that deviates a lot from the movement direction (e.g movement direction going bwd and motion matching could return the fwd animation for a few frames)
-			// When that happens, since we use the delta between root motion and movement direction, we would be over-rotating the lower body and breaking the pose during those frames
-			// So, when that happens we use the inverse of the movement direction to calculate our target rotation. 
-			// This feels a bit 'hacky' but its the only option I've found so far to mitigate the problem
-			if (LocomotionAngleDeltaThreshold > 0.f && FMath::Abs(FMath::RadiansToDegrees(ActualOrientationAngle)) > LocomotionAngleDeltaThreshold)
+			const float RootMotionDeltaSpeed = RootMotionDeltaTranslation.Size() / DeltaSeconds;
+			if (RootMotionDeltaSpeed < MinRootMotionSpeedThreshold)
 			{
-				WarpedRotation = FQuat::FindBetween(RootMotionDeltaDirection, -LocomotionForward);
-				ActualOrientationAngle = WarpedRotation.GetTwistAngle(RotationAxisVector);
+				// If we're under the threshold, snap orientation angle to 0, and let interpolation handle the delta
+				ActualOrientationAngle = 0.0f;
+				PreviousRootMotionDeltaDirection = RootMotionDeltaDirection;
 			}
+			else
+			{
+				// Capture the delta rotation from the axis of motion we care about
+				FQuat WarpedRotation = FQuat::FindBetween(RootMotionDeltaDirection, LocomotionForward);
+
+				ActualOrientationAngle = WarpedRotation.GetTwistAngle(RotationAxisVector);
+				// Motion Matching may return an animation that deviates a lot from the movement direction (e.g movement direction going bwd and motion matching could return the fwd animation for a few frames)
+				// When that happens, since we use the delta between root motion and movement direction, we would be over-rotating the lower body and breaking the pose during those frames
+				// So, when that happens we use the inverse of the movement direction to calculate our target rotation. 
+				// This feels a bit 'hacky' but its the only option I've found so far to mitigate the problem
+				if (LocomotionAngleDeltaThreshold > 0.f && FMath::Abs(FMath::RadiansToDegrees(ActualOrientationAngle)) > LocomotionAngleDeltaThreshold)
+				{
+					WarpedRotation = FQuat::FindBetween(RootMotionDeltaDirection, -LocomotionForward);
+					ActualOrientationAngle = WarpedRotation.GetTwistAngle(RotationAxisVector);
+				}
+
+				// For interpolated warping, guarantee that PreviousOrientationAngle is relative to the current frame's root motion direction 
+				float RootMotionDeltaAngleDifference = FMath::Acos(RootMotionDeltaDirection.Dot(PreviousRootMotionDeltaDirection));
+				RootMotionDeltaAngleDifference *= FMath::Sign(RotationAxisVector.Dot(RootMotionDeltaDirection.Cross(PreviousRootMotionDeltaDirection)));
 
-			// For interpolated warping, guarantee that PreviousOrientationAngle is respect to the current frame's root motion direction 
-			float RootMotionDeltaAngleDifference = FMath::Acos(RootMotionDeltaDirection.Dot(PreviousRootMotionDeltaDirection));
-			RootMotionDeltaAngleDifference *= FMath::Sign(RotationAxisVector.Dot(RootMotionDeltaDirection.Cross(PreviousRootMotionDeltaDirection)));
+				PreviousRootMotionDeltaDirection = RootMotionDeltaDirection;
+				PreviousOrientationAngle += RootMotionDeltaAngleDifference;
+
+				// Rotate the root motion delta fully by the warped angle
+				const FVector WarpedRootMotionTranslationDelta = WarpedRotation.RotateVector(RootMotionDeltaTranslation);
+				RootMotionTransformDelta.SetTranslation(WarpedRootMotionTranslationDelta);
+			}
 
-			PreviousRootMotionDeltaDirection = RootMotionDeltaDirection;
-			PreviousOrientationAngle += RootMotionDeltaAngleDifference;
 
-			// Rotate the root motion delta fully by the warped angle
-			const FVector WarpedRootMotionTranslationDelta = WarpedRotation.RotateVector(RootMotionDeltaTranslation);
-			RootMotionTransformDelta.SetTranslation(WarpedRootMotionTranslationDelta);
 
 			// Forward the side effects of orientation warping on the root motion contribution for this sub-graph
 			const bool bRootMotionOverridden = RootMotionProvider->OverrideRootMotion(RootMotionTransformDelta, Output.CustomAttributes);
@@ -201,7 +215,7 @@ void FAnimNode_OrientationWarping::EvaluateSkeletalControl_AnyThread(FComponentS
 	// Optionally interpolate the effective orientation towards the target orientation angle
 	if (RotationInterpSpeed > 0.f)
 	{
-		ActualOrientationAngle = FMath::FInterpTo(PreviousOrientationAngle, ActualOrientationAngle, Output.AnimInstanceProxy->GetDeltaSeconds(), RotationInterpSpeed);
+		ActualOrientationAngle = FMath::FInterpTo(PreviousOrientationAngle, ActualOrientationAngle, DeltaSeconds, RotationInterpSpeed);
 		PreviousOrientationAngle = ActualOrientationAngle;
 	}
 
diff --git a/Engine/Plugins/Animation/AnimationWarping/Source/Runtime/Public/BoneControllers/AnimNode_OrientationWarping.h b/Engine/Plugins/Animation/AnimationWarping/Source/Runtime/Public/BoneControllers/AnimNode_OrientationWarping.h
--- a/Engine/Plugins/Animation/AnimationWarping/Source/Runtime/Public/BoneControllers/AnimNode_OrientationWarping.h
+++ b/Engine/Plugins/Animation/AnimationWarping/Source/Runtime/Public/BoneControllers/AnimNode_OrientationWarping.h
@@ -24,6 +24,12 @@ struct ANIMATIONWARPINGRUNTIME_API FAnimNode_OrientationWarping : public FAnimNo
 	UPROPERTY(EditAnywhere, BlueprintReadWrite, Category=Evaluation, meta=(PinShownByDefault))
 	float LocomotionAngle = 0.f;
 
+	// Minimum root motion speed required to apply orientation warping
+	// This is useful to prevent unnatural re-orientation when the animation has a portion with no root motion (i.e starts/stops/idles)
+	// When this value is greater than 0, it's recommended to enable interpolation with RotationInterpSpeed > 0
+	UPROPERTY(EditAnywhere, BlueprintReadWrite, Category = Evaluation, meta = (ClampMin = "0.0", PinHiddenByDefault))
+	float MinRootMotionSpeedThreshold = 10.0f;
+
 	// Specifies an angle threshold to prevent erroneous over-rotation of the character, disabled with a value of 0
 	//
 	// When the effective orientation warping angle is detected to be greater than this value (default: 90 degrees) the locomotion direction will be inverted prior to warping
