Change loading phase of AnimationWarping to PreDefault
#rnx
#9217

diff --git a/Engine/Plugins/Animation/AnimationWarping/AnimationWarping.uplugin b/Engine/Plugins/Animation/AnimationWarping/AnimationWarping.uplugin
--- a/Engine/Plugins/Animation/AnimationWarping/AnimationWarping.uplugin
+++ b/Engine/Plugins/Animation/AnimationWarping/AnimationWarping.uplugin
@@ -21,7 +21,7 @@
 		{
 			"Name": "AnimationWarpingEditor",
 			"Type": "UncookedOnly",
-			"LoadingPhase": "Default"
+			"LoadingPhase": "PreDefault"
 		},
 	]
 }
\ No newline at end of file
