Fix for bug when suddenly switching directions with orientation warping

diff --git a/Engine/Plugins/Animation/AnimationWarping/Source/Runtime/Private/BoneControllers/AnimNode_OrientationWarping.cpp b/Engine/Plugins/Animation/AnimationWarping/Source/Runtime/Private/BoneControllers/AnimNode_OrientationWarping.cpp
--- a/Engine/Plugins/Animation/AnimationWarping/Source/Runtime/Private/BoneControllers/AnimNode_OrientationWarping.cpp
+++ b/Engine/Plugins/Animation/AnimationWarping/Source/Runtime/Private/BoneControllers/AnimNode_OrientationWarping.cpp
@@ -155,7 +155,7 @@ void FAnimNode_OrientationWarping::EvaluateSkeletalControl_AnyThread(FComponentS
 			LocomotionForward = SkeletalMeshRelativeRotation.UnrotateVector(LocomotionRotation.GetForwardVector()).GetSafeNormal();
 
 			const FVector RootMotionDeltaTranslation = RootMotionTransformDelta.GetTranslation();
-			// Hold previous direction if we can't calculate it from current move delta
+			// Hold previous direction if we can't calculate it from current move delta, because the root is no longer moving
 			RootMotionDeltaDirection = RootMotionDeltaTranslation.GetSafeNormal(UE_SMALL_NUMBER, PreviousRootMotionDeltaDirection);
 
 			const float RootMotionDeltaSpeed = RootMotionDeltaTranslation.Size() / DeltaSeconds;
@@ -170,24 +170,33 @@ void FAnimNode_OrientationWarping::EvaluateSkeletalControl_AnyThread(FComponentS
 				// Capture the delta rotation from the axis of motion we care about
 				FQuat WarpedRotation = FQuat::FindBetween(RootMotionDeltaDirection, LocomotionForward);
 
+				// For interpolated warping, guarantee that PreviousOrientationAngle is relative to the current frame's root motion direction 
+				float RootMotionDeltaAngleDifference = FMath::Acos(RootMotionDeltaDirection.Dot(PreviousRootMotionDeltaDirection));
+				RootMotionDeltaAngleDifference *= FMath::Sign(RotationAxisVector.Dot(RootMotionDeltaDirection.Cross(PreviousRootMotionDeltaDirection)));
+
+				PreviousRootMotionDeltaDirection = RootMotionDeltaDirection;
+				PreviousOrientationAngle += RootMotionDeltaAngleDifference;
+
 				ActualOrientationAngle = WarpedRotation.GetTwistAngle(RotationAxisVector);
 				// Motion Matching may return an animation that deviates a lot from the movement direction (e.g movement direction going bwd and motion matching could return the fwd animation for a few frames)
 				// When that happens, since we use the delta between root motion and movement direction, we would be over-rotating the lower body and breaking the pose during those frames
 				// So, when that happens we use the inverse of the movement direction to calculate our target rotation. 
 				// This feels a bit 'hacky' but its the only option I've found so far to mitigate the problem
-				if (LocomotionAngleDeltaThreshold > 0.f && FMath::Abs(FMath::RadiansToDegrees(ActualOrientationAngle)) > LocomotionAngleDeltaThreshold)
+				if (LocomotionAngleDeltaThreshold > 0.f)
 				{
-					WarpedRotation = FQuat::FindBetween(RootMotionDeltaDirection, -LocomotionForward);
-					ActualOrientationAngle = WarpedRotation.GetTwistAngle(RotationAxisVector);
+					if (FMath::Abs(FMath::RadiansToDegrees(ActualOrientationAngle)) > LocomotionAngleDeltaThreshold)
+					{
+						WarpedRotation = FQuat::FindBetween(RootMotionDeltaDirection, -LocomotionForward);
+						ActualOrientationAngle = WarpedRotation.GetTwistAngle(RotationAxisVector);
+					}
+					
+					if (FMath::Abs(FMath::RadiansToDegrees(PreviousOrientationAngle)) > LocomotionAngleDeltaThreshold)
+					{
+						// Previous orientation angle might be using an opposite direction too, so flip it if it exceeds the threshold as well.
+						PreviousOrientationAngle = WarpedRotation.GetTwistAngle(RotationAxisVector);
+					}
 				}
 
-				// For interpolated warping, guarantee that PreviousOrientationAngle is relative to the current frame's root motion direction 
-				float RootMotionDeltaAngleDifference = FMath::Acos(RootMotionDeltaDirection.Dot(PreviousRootMotionDeltaDirection));
-				RootMotionDeltaAngleDifference *= FMath::Sign(RotationAxisVector.Dot(RootMotionDeltaDirection.Cross(PreviousRootMotionDeltaDirection)));
-
-				PreviousRootMotionDeltaDirection = RootMotionDeltaDirection;
-				PreviousOrientationAngle += RootMotionDeltaAngleDifference;
-
 				// Rotate the root motion delta fully by the warped angle
 				const FVector WarpedRootMotionTranslationDelta = WarpedRotation.RotateVector(RootMotionDeltaTranslation);
 				RootMotionTransformDelta.SetTranslation(WarpedRootMotionTranslationDelta);
