Foot Placement:
-Fix for nans when nearing infinity
-Properly handle teleports

diff --git a/Engine/Plugins/Animation/AnimationWarping/Source/Runtime/Private/BoneControllers/AnimNode_FootPlacement.cpp b/Engine/Plugins/Animation/AnimationWarping/Source/Runtime/Private/BoneControllers/AnimNode_FootPlacement.cpp
--- a/Engine/Plugins/Animation/AnimationWarping/Source/Runtime/Private/BoneControllers/AnimNode_FootPlacement.cpp
+++ b/Engine/Plugins/Animation/AnimationWarping/Source/Runtime/Private/BoneControllers/AnimNode_FootPlacement.cpp
@@ -103,6 +103,18 @@ namespace UE::Anim::FootPlacement
 		return InNormal;
 	}
 
+	// Since we are calculating in world-space, when too far from the origin FMath::LinePlaneIntersection can introduce numerical error
+	// and consider the line's start/end to be at the same location.
+	// Use point-direction instead to avoid this
+	// See UE-162275
+	static FVector PointDirectionPlaneIntersection(const FVector Point, const FVector Direction, const FPlane Plane)
+	{
+		return
+			Point
+			+ Direction
+			* ((Plane.W - (Point | Plane)) / (Direction | Plane));
+	};
+
 	static bool FindPlantTraceImpact(
 		const FEvaluationContext& Context,
 		const FFootPlacementTraceSettings& TraceSettings,
@@ -187,9 +199,9 @@ namespace UE::Anim::FootPlacement
 
 	static float GetDistanceToPlaneAlongDirection(const FVector& Location, const FPlane& PlantPlane, const FVector& ApproachDir)
 	{
-		const FVector IntersectionLoc = FMath::LinePlaneIntersection(
+		const FVector IntersectionLoc = UE::Anim::FootPlacement::PointDirectionPlaneIntersection(
 			Location,
-			Location - ApproachDir,
+			-ApproachDir,
 			PlantPlane);
 
 		const FVector IntersectionToLocation = Location - IntersectionLoc;
@@ -364,9 +376,9 @@ void FAnimNode_FootPlacement::AlignPlantToGround(
 	const float IKFootRootToFootRootTargetDistance = 
 		UE::Anim::FootPlacement::GetDistanceToPlaneAlongDirection(InputPoseFootTransformWS.GetLocation(), IKFootRootPlaneWS, Context.ApproachDirWS);
 
-	const FVector CorrectedPlaneIntersectionWS = FMath::LinePlaneIntersection(
+	const FVector CorrectedPlaneIntersectionWS = UE::Anim::FootPlacement::PointDirectionPlaneIntersection(
 		InOutFootTransformWS.GetLocation(),
-		InOutFootTransformWS.GetLocation() + Context.ApproachDirWS,
+		Context.ApproachDirWS,
 		PlantPlaneWS);
 
 	const FVector CorrectedLocationWS =
@@ -478,20 +490,19 @@ void FAnimNode_FootPlacement::UpdatePlantingPlaneInterpolation(
 		UE::Anim::FootPlacement::FindPlantPlane(Context, TraceSettings, FootTransformWS.GetLocation(), true, InOutPlantPlane, ImpactLocationWS);
 	}
 	
-
-	const FVector CurrPlaneIntersection = FMath::LinePlaneIntersection(
+	const FVector CurrPlaneIntersection = UE::Anim::FootPlacement::PointDirectionPlaneIntersection(
 		FootTransformWS.GetLocation(),
-		FootTransformWS.GetLocation() + TraceDirection,
+		TraceDirection,
 		InOutPlantPlane);
 
-	const FVector LastPlaneIntersection = FMath::LinePlaneIntersection(
+	const FVector LastPlaneIntersection = UE::Anim::FootPlacement::PointDirectionPlaneIntersection(
 		LastAlignedFootTransform.GetLocation(),
-		LastAlignedFootTransform.GetLocation() + TraceDirection,
+		TraceDirection,
 		LastPlantPlane);
 
-	const FVector PrevPlaneIntersection = FMath::LinePlaneIntersection(
+	const FVector PrevPlaneIntersection = UE::Anim::FootPlacement::PointDirectionPlaneIntersection(
 		FootTransformWS.GetLocation(),
-		FootTransformWS.GetLocation() + TraceDirection,
+		TraceDirection,
 		LastPlantPlane);
 
 	const float LastPlaneDeltaZ = LastPlaneIntersection.Z - CurrPlaneIntersection.Z;
@@ -609,6 +620,27 @@ float FAnimNode_FootPlacement::GetMinLimbExtension(const float DesiredExtension,
 	return FMath::Min(DesiredExtension, LimbLength * PlantSettings.MinExtensionRatio);
 }
 
+void FAnimNode_FootPlacement::ResetRuntimeData() 
+{
+	PelvisData.Interpolation = UE::Anim::FootPlacement::FPelvisRuntimeData::FInterpolationData();
+
+	LegsData.Reset();
+	LegsData.SetNumUninitialized(LegDefinitions.Num());
+
+	for (int32 LegIndex = 0; LegIndex < LegsData.Num(); ++LegIndex)
+	{
+		UE::Anim::FootPlacement::FLegRuntimeData& LegData = LegsData[LegIndex];
+		LegData.Idx = LegIndex;
+		LegData.Interpolation = UE::Anim::FootPlacement::FLegRuntimeData::FInterpolationData();
+	}
+
+#if ENABLE_ANIM_DEBUG
+	DebugData.Init(LegDefinitions.Num());
+#endif
+
+	bIsFirstUpdate = true;
+}
+
 bool FAnimNode_FootPlacement::WantsToPlant(
 	const UE::Anim::FootPlacement::FEvaluationContext& Context,
 	const UE::Anim::FootPlacement::FLegRuntimeData::FInputPoseData& LegInputPose) const
@@ -803,6 +835,8 @@ UE::Anim::FootPlacement::FPlantResult FAnimNode_FootPlacement::FinalizeFootAlign
 		}
 	}
 
+	check(!CorrectedFootTransformCS.ContainsNaN());
+
 	// TODO: Do adjustments to FK tip and FK Chain root
 	FTransform FinalFkTipTransformCS = FTransform::Identity;
 	FTransform FinalFkHipTransformCS = FTransform::Identity;
@@ -849,9 +883,9 @@ void FAnimNode_FootPlacement::DrawDebug(
 		LegData.AlignedFootTransformWS;
 
 
-	const FVector FKBoneLocationProjectedWS = FMath::LinePlaneIntersection(
+	const FVector FKBoneLocationProjectedWS = UE::Anim::FootPlacement::PointDirectionPlaneIntersection(
 		FKBoneTransformWS.GetLocation(),
-		FKBoneTransformWS.GetLocation() + Context.ApproachDirWS,
+		Context.ApproachDirWS,
 		LegData.Plant.PlantPlaneWS);
 
 	Context.CSPContext.AnimInstanceProxy->AnimDrawDebugPoint(
@@ -865,9 +899,9 @@ void FAnimNode_FootPlacement::DrawDebug(
 	Context.CSPContext.AnimInstanceProxy->AnimDrawDebugPoint(
 		IKBoneTransformWS.GetLocation(), 10.0f, CurrentPlantColor, false, -1.0f, SDPG_Foreground);
 
-	const FVector IKBoneLocationProjectedWS = FMath::LinePlaneIntersection(
+	const FVector IKBoneLocationProjectedWS = UE::Anim::FootPlacement::PointDirectionPlaneIntersection(
 		IKBoneTransformWS.GetLocation(),
-		IKBoneTransformWS.GetLocation() + Context.ApproachDirWS,
+		Context.ApproachDirWS,
 		LegData.Plant.PlantPlaneWS);
 
 	Context.CSPContext.AnimInstanceProxy->AnimDrawDebugPoint(
@@ -878,9 +912,9 @@ void FAnimNode_FootPlacement::DrawDebug(
 		false, -1.0f, 1.0f, SDPG_Foreground);
 
 	const float UnplantRadius = PlantSettings.UnplantRadius;
-	const FVector PlantCenter = FMath::LinePlaneIntersection(
+	const FVector PlantCenter = UE::Anim::FootPlacement::PointDirectionPlaneIntersection(
 		IKBoneTransformWS.GetLocation(),
-		IKBoneTransformWS.GetLocation() + Context.ApproachDirWS,
+		Context.ApproachDirWS,
 		LegData.Plant.PlantPlaneWS);
 	Context.CSPContext.AnimInstanceProxy->AnimDrawDebugCircle(
 		PlantCenter, UnplantRadius, 24, PlantedColor,
@@ -943,24 +977,7 @@ void FAnimNode_FootPlacement::GatherDebugData(FNodeDebugData& NodeDebugData)
 void FAnimNode_FootPlacement::Initialize_AnyThread(const FAnimationInitializeContext& Context)
 {
 	FAnimNode_SkeletalControlBase::Initialize_AnyThread(Context);
-
-	PelvisData.Interpolation = UE::Anim::FootPlacement::FPelvisRuntimeData::FInterpolationData();
-
-	LegsData.Reset();
-	LegsData.SetNumUninitialized(LegDefinitions.Num());
-
-	for (int32 LegIndex = 0; LegIndex < LegsData.Num(); ++LegIndex)
-	{
-		UE::Anim::FootPlacement::FLegRuntimeData& LegData = LegsData[LegIndex];
-		LegData.Idx = LegIndex;
-		LegData.Interpolation = UE::Anim::FootPlacement::FLegRuntimeData::FInterpolationData();
-	}
-
-#if ENABLE_ANIM_DEBUG
-	DebugData.Init(LegDefinitions.Num());
-#endif
-
-	bIsFirstUpdate = true;
+	ResetRuntimeData();
 }
 
 void FAnimNode_FootPlacement::UpdateInternal(const FAnimationUpdateContext& Context)
@@ -970,8 +987,7 @@ void FAnimNode_FootPlacement::UpdateInternal(const FAnimationUpdateContext& Cont
 	// If we just became relevant and haven't been initialized yet, then reinitialize foot placement.
 	if (!bIsFirstUpdate && UpdateCounter.HasEverBeenUpdated() && !UpdateCounter.WasSynchronizedCounter(Context.AnimInstanceProxy->GetUpdateCounter()))
 	{
-		FAnimationInitializeContext InitializationContext(Context.AnimInstanceProxy, Context.SharedContext);
-		Initialize_AnyThread(InitializationContext);
+		ResetRuntimeData();
 	}
 	UpdateCounter.SynchronizeWith(Context.AnimInstanceProxy->GetUpdateCounter());
 
@@ -985,6 +1001,22 @@ void FAnimNode_FootPlacement::EvaluateSkeletalControl_AnyThread(FComponentSpaceP
 
 	check(OutBoneTransforms.Num() == 0);
 
+	// Manually calculate distance instead of using the teleport flag to properly handle cases like crouch
+	// and instantaneous root offsets i.e. when entering/leaving a vehicle.
+	// See FAnimNode_Inertialization::Evaluate_AnyThread and/or UE-78594
+	const float TeleportDistanceThreshold = Output.AnimInstanceProxy->GetSkelMeshComponent()->GetTeleportDistanceThreshold();
+	if (!bIsFirstUpdate && (TeleportDistanceThreshold > 0.0f))
+	{
+		const FTransform ComponentTransform = Output.AnimInstanceProxy->GetComponentTransform();
+		const FVector RootWorldSpaceLocation = ComponentTransform.TransformPosition(
+			Output.Pose.GetComponentSpaceTransform(FCompactPoseBoneIndex(0)).GetTranslation());
+		const FVector PrevRootWorldSpaceLocation = CharacterData.ComponentTransformWS.TransformPosition(PelvisData.InputPose.RootTransformCS.GetTranslation());
+		if (FVector::DistSquared(RootWorldSpaceLocation, PrevRootWorldSpaceLocation) > FMath::Square(TeleportDistanceThreshold))
+		{
+			ResetRuntimeData();
+		}
+	}
+	
 #if ENABLE_ANIM_DEBUG
 	UE::Anim::FootPlacement::FDebugData LastDebugData = DebugData;
 #endif
@@ -1026,6 +1058,7 @@ void FAnimNode_FootPlacement::EvaluateSkeletalControl_AnyThread(FComponentSpaceP
 	// Based on the ground alignment, search for the best Pelvis transform
 	FTransform PelvisTransformCS = SolvePelvis(FootPlacementContext);
 	PelvisTransformCS = UpdatePelvisInterpolation(FootPlacementContext, PelvisTransformCS);
+	check(!PelvisTransformCS.ContainsNaN());
 	OutBoneTransforms.Add(FBoneTransform(PelvisData.Bones.FkBoneIndex, PelvisTransformCS));
 
 #if ENABLE_ANIM_DEBUG
@@ -1200,6 +1233,8 @@ void FAnimNode_FootPlacement::GatherPelvisDataFromInputs(const UE::Anim::FootPla
 		Context.CSPContext.Pose.GetComponentSpaceTransform(PelvisData.Bones.FkBoneIndex);
 	PelvisData.InputPose.IKRootTransformCS =
 		Context.CSPContext.Pose.GetComponentSpaceTransform(PelvisData.Bones.IkBoneIndex);
+	PelvisData.InputPose.RootTransformCS =
+		Context.CSPContext.Pose.GetComponentSpaceTransform(FCompactPoseBoneIndex(0));
 
 	// TODO: All of these can be calculated on initialize, but in case there's value in changing these dynamically,
 	// will keep this for now. If needed, change to lazy update.
@@ -1285,9 +1320,10 @@ void FAnimNode_FootPlacement::ProcessCharacterState(const UE::Anim::FootPlacemen
 {
 	const FVector LastComponentLocationWS = bIsFirstUpdate
 		? Context.OwningComponentToWorld.GetLocation()
-		: CharacterData.ComponentLocationWS;
+		: CharacterData.ComponentTransformWS.GetLocation();
 
-	CharacterData.ComponentLocationWS = Context.OwningComponentToWorld.GetLocation();
+	CharacterData.ComponentTransformWS = Context.OwningComponentToWorld;
+	const FVector ComponentLocationWS = CharacterData.ComponentTransformWS.GetLocation();
 	CharacterData.NumFKPlanted = 0;
 	for (const UE::Anim::FootPlacement::FLegRuntimeData& LegData : LegsData)
 	{
@@ -1310,14 +1346,14 @@ void FAnimNode_FootPlacement::ProcessCharacterState(const UE::Anim::FootPlacemen
 		const FVector CapsuleFloorNormalWS = Context.GetMovementComponentFloorNormal();
 		const FVector OwningComponentAdjustedLastLocationWS =
 			(FMath::Abs(Context.ApproachDirWS | CapsuleFloorNormalWS) > DELTA) ?
-			FMath::LinePlaneIntersection(
-				CharacterData.ComponentLocationWS,
-				CharacterData.ComponentLocationWS + Context.ApproachDirWS,
-				LastComponentLocationWS, CapsuleFloorNormalWS) :
-			CharacterData.ComponentLocationWS;
+			UE::Anim::FootPlacement::PointDirectionPlaneIntersection(
+				ComponentLocationWS,
+				Context.ApproachDirWS,
+				FPlane(LastComponentLocationWS, CapsuleFloorNormalWS)) :
+			ComponentLocationWS;
 
 		const FVector CapsuleMoveOffsetWS =
-			CharacterData.ComponentLocationWS - OwningComponentAdjustedLastLocationWS;
+			ComponentLocationWS - OwningComponentAdjustedLastLocationWS;
 		if (!CapsuleMoveOffsetWS.IsNearlyZero(KINDA_SMALL_NUMBER))
 		{
 			const FVector CapsuleMoveOffsetCS =
diff --git a/Engine/Plugins/Animation/AnimationWarping/Source/Runtime/Public/BoneControllers/AnimNode_FootPlacement.h b/Engine/Plugins/Animation/AnimationWarping/Source/Runtime/Public/BoneControllers/AnimNode_FootPlacement.h
--- a/Engine/Plugins/Animation/AnimationWarping/Source/Runtime/Public/BoneControllers/AnimNode_FootPlacement.h
+++ b/Engine/Plugins/Animation/AnimationWarping/Source/Runtime/Public/BoneControllers/AnimNode_FootPlacement.h
@@ -112,6 +112,7 @@ namespace UE::Anim::FootPlacement
 		{
 			FTransform FKTransformCS = FTransform::Identity;
 			FTransform IKRootTransformCS = FTransform::Identity;
+			FTransform RootTransformCS = FTransform::Identity;
 			FVector FootMidpointCS = FVector::ZeroVector;
 		} InputPose;
 
@@ -125,7 +126,7 @@ namespace UE::Anim::FootPlacement
 
 	struct FCharacterData
 	{
-		FVector ComponentLocationWS = FVector::ZeroVector;
+		FTransform ComponentTransformWS = FTransform::Identity;
 		int NumFKPlanted = 0;
 		bool bIsOnGround = false;
 	};
@@ -568,6 +569,8 @@ private:
 	float GetMaxLimbExtension(const float DesiredExtension, const float LimbLength) const;
 	float GetMinLimbExtension(const float DesiredExtension, const float LimbLength) const;
 
+	void ResetRuntimeData();
+
 #if ENABLE_ANIM_DEBUG
 	UE::Anim::FootPlacement::FDebugData DebugData;
 	void DrawDebug(
