Fix for foot placement popping the feet when used along other warping nodes

diff --git a/Engine/Plugins/Animation/AnimationWarping/Source/Runtime/Private/BoneControllers/AnimNode_FootPlacement.cpp b/Engine/Plugins/Animation/AnimationWarping/Source/Runtime/Private/BoneControllers/AnimNode_FootPlacement.cpp
--- a/Engine/Plugins/Animation/AnimationWarping/Source/Runtime/Private/BoneControllers/AnimNode_FootPlacement.cpp
+++ b/Engine/Plugins/Animation/AnimationWarping/Source/Runtime/Private/BoneControllers/AnimNode_FootPlacement.cpp
@@ -254,7 +254,7 @@ void FAnimNode_FootPlacement::FindPelvisOffsetRangeForLimb(
 	const FVector HipLocationCS = HipTransformCS.GetLocation();
 
 	const FVector DesiredExtensionDelta =
-		LegInputPose.FKTransformCS.GetLocation() - LegInputPose.HipTransformCS.GetLocation();
+		LegInputPose.FootTransformCS.GetLocation() - LegInputPose.HipTransformCS.GetLocation();
 
 	const float DesiredExtensionSqrd = DesiredExtensionDelta.SizeSquared();
 	const float DesiredExtension = FMath::Sqrt(DesiredExtensionSqrd);
@@ -274,7 +274,7 @@ void FAnimNode_FootPlacement::FindPelvisOffsetRangeForLimb(
 		const float MaxExtensionSqrd = MaxExtension * MaxExtension;
 
 		const FPlane FootPlane = FPlane(PlantTargetLocationCS, Context.ApproachDirCS);
-		const FVector FKFootProjected = FVector::PointPlaneProject(LegInputPose.FKTransformCS.GetLocation(), FootPlane);
+		const FVector FKFootProjected = FVector::PointPlaneProject(LegInputPose.FootTransformCS.GetLocation(), FootPlane);
 		const FVector HipProjected = FVector::PointPlaneProject(HipLocationCS, FootPlane);
 
 		const float MaxOffset = PelvisSettings.MaxOffsetHorizontal;
@@ -341,7 +341,7 @@ float FAnimNode_FootPlacement::CalcTargetPlantPlaneDistance(
 	const UE::Anim::FootPlacement::FEvaluationContext& Context,
 	const UE::Anim::FootPlacement::FLegRuntimeData::FInputPoseData& LegInputPose) const
 {
-	const FTransform IKBallBoneCS = LegInputPose.FootToBall * LegInputPose.IKTransformCS;
+	const FTransform IKBallBoneCS = LegInputPose.FootToBall * LegInputPose.FootTransformCS;
 
 	const FTransform& IKFootRootCS = PelvisData.InputPose.IKRootTransformCS;
 	const FPlane IKGroundPlaneCS = 
@@ -350,7 +350,7 @@ float FAnimNode_FootPlacement::CalcTargetPlantPlaneDistance(
 
 	// TODO: I'm just getting the distance between bones and the plane, instead of actual foot/ball bases
 	const float FootBaseDistance = 
-		UE::Anim::FootPlacement::GetDistanceToPlaneAlongDirection(LegInputPose.IKTransformCS.GetLocation(), IKGroundPlaneCS, Context.ApproachDirCS);
+		UE::Anim::FootPlacement::GetDistanceToPlaneAlongDirection(LegInputPose.FootTransformCS.GetLocation(), IKGroundPlaneCS, Context.ApproachDirCS);
 	const float BallBaseDistance = 
 		UE::Anim::FootPlacement::GetDistanceToPlaneAlongDirection(IKBallBoneCS.GetLocation(), IKGroundPlaneCS, Context.ApproachDirCS);
 
@@ -366,7 +366,7 @@ void FAnimNode_FootPlacement::AlignPlantToGround(
 	FTransform& InOutFootTransformWS,
 	FQuat& OutTwistCorrection) const
 {
-	const FTransform InputPoseFootTransformWS = LegInputPose.IKTransformCS * Context.OwningComponentToWorld;
+	const FTransform InputPoseFootTransformWS = LegInputPose.FootTransformCS * Context.OwningComponentToWorld;
 
 	// It is assumed the distance from the plane defined by ik foot root to the ik reference, along the trace 
 	// direction, must remain the same.
@@ -703,7 +703,7 @@ UE::Anim::FootPlacement::FPlantResult FAnimNode_FootPlacement::FinalizeFootAlign
 
 	// avoid hyper extension - start
 	const FVector InitialHipToFootDir =
-		(LegData.InputPose.FKTransformCS.GetLocation() - LegData.InputPose.HipTransformCS.GetLocation()).GetSafeNormal();
+		(LegData.InputPose.FootTransformCS.GetLocation() - LegData.InputPose.HipTransformCS.GetLocation()).GetSafeNormal();
 	const FVector TargetHipToFootDir =
 		(CorrectedFootTransformCS.GetLocation() - FinalHipTransformCS.GetLocation()).GetSafeNormal();
 
@@ -711,7 +711,7 @@ UE::Anim::FootPlacement::FPlantResult FAnimNode_FootPlacement::FinalizeFootAlign
 	LegData.Plant.bCanReachTarget = true;
 
 	const FVector FKHipToFoot =
-		LegData.InputPose.FKTransformCS.GetLocation() - LegData.InputPose.HipTransformCS.GetLocation();
+		LegData.InputPose.FootTransformCS.GetLocation() - LegData.InputPose.HipTransformCS.GetLocation();
 
 	if (!InitialHipToFootDir.IsNearlyZero() && !TargetHipToFootDir.IsNearlyZero())
 	{
@@ -739,7 +739,7 @@ UE::Anim::FootPlacement::FPlantResult FAnimNode_FootPlacement::FinalizeFootAlign
 				}
 
 				const FTransform FKHipToLeg =
-					LegData.InputPose.FKTransformCS.GetRelativeTransform(PelvisData.InputPose.FKTransformCS);
+					LegData.InputPose.FootTransformCS.GetRelativeTransform(PelvisData.InputPose.FKTransformCS);
 				const FTransform FKLegAtCurrentHipCS =
 					FKHipToLeg * PelvisTransformCS;
 				FVector IKToFK = (CorrectedFootTransformCS.GetLocation() - FKLegAtCurrentHipCS.GetLocation()).GetSafeNormal();
@@ -875,7 +875,7 @@ void FAnimNode_FootPlacement::DrawDebug(
 
 	const FTransform FKBoneTransformWS =
 		LegData.InputPose.FootToGround *
-		LegData.InputPose.FKTransformCS *
+		LegData.InputPose.FootTransformCS *
 		Context.OwningComponentToWorld;
 
 	const FTransform IKBoneTransformWS =
@@ -1089,7 +1089,7 @@ void FAnimNode_FootPlacement::EvaluateSkeletalControl_AnyThread(FComponentSpaceP
 			DebugData.OutputFootLocationsWS[FootIndex] =
 				FootPlacementContext.OwningComponentToWorld.TransformPosition(PlantResult.IkPlantTranformCS.Transform.GetLocation());
 			DebugData.InputFootLocationsWS[FootIndex] = 
-				FootPlacementContext.OwningComponentToWorld.TransformPosition(LegData.InputPose.IKTransformCS.GetLocation());;
+				FootPlacementContext.OwningComponentToWorld.TransformPosition(LegData.InputPose.FootTransformCS.GetLocation());;
 		}
 #endif
 	}
@@ -1249,24 +1249,30 @@ void FAnimNode_FootPlacement::GatherLegDataFromInputs(
 {
 	FVector LastBallLocation = LegData.InputPose.BallTransformCS.GetLocation();
 
-	LegData.InputPose.FKTransformCS =
+	const FTransform FootFKTransformCS =
 		Context.CSPContext.Pose.GetComponentSpaceTransform(LegData.Bones.FKIndex);
-	LegData.InputPose.IKTransformCS =
-		Context.CSPContext.Pose.GetComponentSpaceTransform(LegData.Bones.IKIndex);
-	LegData.InputPose.BallTransformCS =
+	const FTransform BallTransformCS =
 		Context.CSPContext.Pose.GetComponentSpaceTransform(LegData.Bones.BallIndex);
+
+	LegData.InputPose.FootTransformCS =
+		Context.CSPContext.Pose.GetComponentSpaceTransform(LegData.Bones.IKIndex);
 	LegData.InputPose.HipTransformCS =
 		Context.CSPContext.Pose.GetComponentSpaceTransform(LegData.Bones.HipIndex);
 
 	LegData.InputPose.BallToFoot =
-		LegData.InputPose.FKTransformCS.GetRelativeTransform(LegData.InputPose.BallTransformCS);
+		FootFKTransformCS.GetRelativeTransform(BallTransformCS);
 	LegData.InputPose.FootToBall =
-		LegData.InputPose.BallTransformCS.GetRelativeTransform(LegData.InputPose.FKTransformCS);
+		BallTransformCS.GetRelativeTransform(FootFKTransformCS);
+
+	// Can't use ball transform as-is as the foot's IK bone may not be at the FK bone.
+	// Assume the ball is at the same relative position
+	LegData.InputPose.BallTransformCS =
+		LegData.InputPose.FootToBall * LegData.InputPose.FootTransformCS;
 
 	if (bIsFirstUpdate)
 	{
 		LegData.AlignedFootTransformWS =
-			LegData.InputPose.FKTransformCS * Context.OwningComponentToWorld;
+			LegData.InputPose.FootTransformCS * Context.OwningComponentToWorld;
 		LegData.UnalignedFootTransformWS = LegData.AlignedFootTransformWS;
 
 		const FVector IKFootRootLocationWS =
@@ -1312,7 +1318,7 @@ void FAnimNode_FootPlacement::CalculateFootMidpoint(const UE::Anim::FootPlacemen
 	OutMidpoint = FVector::ZeroVector;
 	for (const UE::Anim::FootPlacement::FLegRuntimeData& LegData : InLegsData)
 	{
-		OutMidpoint += LegData.InputPose.IKTransformCS.GetLocation() / NumLegs;
+		OutMidpoint += LegData.InputPose.FootTransformCS.GetLocation() / NumLegs;
 	}
 }
 
@@ -1373,15 +1379,14 @@ void FAnimNode_FootPlacement::ProcessFootAlignment(
 	UE::Anim::FootPlacement::FLegRuntimeData::FBoneData& Bones = LegData.Bones;
 	UE::Anim::FootPlacement::FLegRuntimeData::FPlantData& Plant = LegData.Plant;
 
-	const FTransform FKFootTransformWS = InputPose.FKTransformCS * Context.OwningComponentToWorld;
-	const FTransform IKFootTransformWS = InputPose.IKTransformCS * Context.OwningComponentToWorld;
+	const FTransform InputPoseFootTransformWS = InputPose.FootTransformCS * Context.OwningComponentToWorld;
 	const FTransform LastAlignedFootTransformWS = LegData.AlignedFootTransformWS;
 	const FTransform LastUnalignedFootTransformWS = LegData.UnalignedFootTransformWS;
 
 	Plant.LastPlantType = Plant.PlantType;
 	DeterminePlantType(
 		Context,
-		FKFootTransformWS,
+		InputPoseFootTransformWS,
 		LastAlignedFootTransformWS,
 		Plant,
 		InputPose);
@@ -1406,7 +1411,7 @@ void FAnimNode_FootPlacement::ProcessFootAlignment(
 		case EFootPlacementLockType::PivotAroundAnkle:
 		{
 			// Use the location only
-			CurrentPlantedTransformWS = IKFootTransformWS;
+			CurrentPlantedTransformWS = InputPoseFootTransformWS;
 			CurrentPlantedTransformWS.SetLocation(LastUnalignedFootTransformWS.GetLocation());
 		}
 			break;
@@ -1425,7 +1430,7 @@ void FAnimNode_FootPlacement::ProcessFootAlignment(
 
 		// The locked transform is aligned to the ground. Conserve the input pose's ground alignment
 		const FVector AlignedBoneLocationCS = PlantedFootTransformCS.GetLocation();
-		const FPlane InputPosePlantPlane = FPlane(InputPose.IKTransformCS.GetLocation(), Context.ApproachDirWS);
+		const FPlane InputPosePlantPlane = FPlane(InputPose.FootTransformCS.GetLocation(), Context.ApproachDirWS);
 		const FVector UnalignedBoneLocationCS = FVector::PointPlaneProject(AlignedBoneLocationCS, InputPosePlantPlane);
 
 		PlantedFootTransformCS.SetLocation(UnalignedBoneLocationCS);
@@ -1433,7 +1438,7 @@ void FAnimNode_FootPlacement::ProcessFootAlignment(
 		// Get the offset relative to the initial foot transform
 		// Reset interpolation 
 		Interpolation.UnalignedFootOffsetCS =
-			InputPose.IKTransformCS.GetRelativeTransformReverse(PlantedFootTransformCS);
+			InputPose.FootTransformCS.GetRelativeTransformReverse(PlantedFootTransformCS);
 		Interpolation.PlantOffsetTranslationSpringState.Reset();
 		Interpolation.PlantOffsetRotationSpringState.Reset();
 
@@ -1444,7 +1449,7 @@ void FAnimNode_FootPlacement::ProcessFootAlignment(
 	{
 		// No plant, so we interpolate the offset out
 		Interpolation.UnalignedFootOffsetCS =
-			UpdatePlantOffsetInterpolation(Context, Interpolation, InputPose.IKTransformCS);
+			UpdatePlantOffsetInterpolation(Context, Interpolation, InputPose.FootTransformCS);
 
 		// If we're unplanted, we know we're fully unaligned the first time we hit zero alignment alpha.
 		if (Plant.TimeSinceFullyUnaligned > 0.0f || FMath::IsNearlyZero(InputPose.AlignmentAlpha))
@@ -1479,23 +1484,23 @@ void FAnimNode_FootPlacement::ProcessFootAlignment(
 		Interpolation.UnalignedFootOffsetCS.SetRotation(ClampedRotationOffset);
 	}
 
-	FTransform IKUnalignedTransformCS = InputPose.IKTransformCS * Interpolation.UnalignedFootOffsetCS;
+	FTransform FootUnalignedTransformCS = InputPose.FootTransformCS * Interpolation.UnalignedFootOffsetCS;
 	if (PlantSettings.SeparatingDistance > 0.0f)
 	{
 		// Prevent the feet from crossing by enforcing a set distance from a plane at the midpoint between all feet
-		FVector IKUnalignedLocationCS = IKUnalignedTransformCS.GetLocation();
-		const FVector MidPointToFoot = (InputPose.IKTransformCS.GetLocation() - PelvisData.InputPose.FootMidpointCS);
+		FVector FootUnalignedLocationCS = FootUnalignedTransformCS.GetLocation();
+		const FVector MidPointToFoot = (InputPose.FootTransformCS.GetLocation() - PelvisData.InputPose.FootMidpointCS);
 		const FVector PlaneNormal = MidPointToFoot.GetSafeNormal2D();
 		const FVector PlaneCenter = PelvisData.InputPose.FootMidpointCS + PlaneNormal * PlantSettings.SeparatingDistance;
 		const FPlane SeparatingPlane = FPlane(PlaneCenter, PlaneNormal);
 		
 		const float DistanceToSeparatingPlane =
-			UE::Anim::FootPlacement::GetDistanceToPlaneAlongDirection(IKUnalignedLocationCS, SeparatingPlane, -PlaneNormal);
+			UE::Anim::FootPlacement::GetDistanceToPlaneAlongDirection(FootUnalignedLocationCS, SeparatingPlane, -PlaneNormal);
 
 		if (DistanceToSeparatingPlane < 0.0f)
 		{
-			IKUnalignedLocationCS -= PlaneNormal * DistanceToSeparatingPlane;
-			IKUnalignedTransformCS.SetLocation(IKUnalignedLocationCS);
+			FootUnalignedLocationCS -= PlaneNormal * DistanceToSeparatingPlane;
+			FootUnalignedTransformCS.SetLocation(FootUnalignedLocationCS);
 		}
 
 #if ENABLE_ANIM_DEBUG
@@ -1512,7 +1517,7 @@ void FAnimNode_FootPlacement::ProcessFootAlignment(
 #endif
 	}
 
-	LegData.UnalignedFootTransformWS = IKUnalignedTransformCS * Context.OwningComponentToWorld;
+	LegData.UnalignedFootTransformWS = FootUnalignedTransformCS * Context.OwningComponentToWorld;
 
 	const FTransform ComponentToWorldInv = Context.OwningComponentToWorld.Inverse();
 
@@ -1520,9 +1525,9 @@ void FAnimNode_FootPlacement::ProcessFootAlignment(
 	UpdatePlantingPlaneInterpolation(Context, LegData.UnalignedFootTransformWS, LastAlignedFootTransformWS, InputPose.AlignmentAlpha, Plant.PlantPlaneWS, Interpolation);
 	Plant.PlantPlaneCS = Plant.PlantPlaneWS.TransformBy(ComponentToWorldInv.ToMatrixWithScale());
 
-	Interpolation.UnalignedFootOffsetCS = InputPose.IKTransformCS.GetRelativeTransformReverse(IKUnalignedTransformCS);
+	Interpolation.UnalignedFootOffsetCS = InputPose.FootTransformCS.GetRelativeTransformReverse(FootUnalignedTransformCS);
 
-	// This will adjust IkUnalignedTransformWS to make it match the required distance to the plant plane along the 
+	// This will adjust UnalignedFootTransformWS to make it match the required distance to the plant plane along the 
 	// approach direction, not the plane normal
 	LegData.AlignedFootTransformWS = LegData.UnalignedFootTransformWS;
 	AlignPlantToGround(Context, Plant.PlantPlaneWS, InputPose, LegData.AlignedFootTransformWS, Plant.TwistCorrection);
@@ -1538,7 +1543,7 @@ void FAnimNode_FootPlacement::ProcessFootAlignment(
 	// When unplanted/unaligned, favor FK orientation and fix penetrations later. 
 	BlendedPlantTransformCS.SetRotation(
 		FQuat::Slerp(
-			InputPose.FKTransformCS.GetRotation(),
+			InputPose.FootTransformCS.GetRotation(),
 			LegData.AlignedFootTransformCS.GetRotation(),
 			InputPose.AlignmentAlpha));
 
@@ -1607,7 +1612,7 @@ FTransform FAnimNode_FootPlacement::SolvePelvis(const UE::Anim::FootPlacement::F
 		FVector OffsetAverage = FVector::ZeroVector;
 		for (const FLegRuntimeData& LegData : LegsData)
 		{
-			const FVector LegTranslationOffset = LegData.AlignedFootTransformCS.GetLocation() - LegData.InputPose.IKTransformCS.GetLocation();
+			const FVector LegTranslationOffset = LegData.AlignedFootTransformCS.GetLocation() - LegData.InputPose.FootTransformCS.GetLocation();
 			OffsetAverage += LegTranslationOffset / NumLegs;
 		}
 
diff --git a/Engine/Plugins/Animation/AnimationWarping/Source/Runtime/Public/BoneControllers/AnimNode_FootPlacement.h b/Engine/Plugins/Animation/AnimationWarping/Source/Runtime/Public/BoneControllers/AnimNode_FootPlacement.h
--- a/Engine/Plugins/Animation/AnimationWarping/Source/Runtime/Public/BoneControllers/AnimNode_FootPlacement.h
+++ b/Engine/Plugins/Animation/AnimationWarping/Source/Runtime/Public/BoneControllers/AnimNode_FootPlacement.h
@@ -38,8 +38,7 @@ namespace UE::Anim::FootPlacement
 		// Helper struct to store values used across the evaluation
 		struct FInputPoseData
 		{
-			FTransform IKTransformCS = FTransform::Identity;
-			FTransform FKTransformCS = FTransform::Identity;
+			FTransform FootTransformCS = FTransform::Identity;
 			FTransform BallTransformCS = FTransform::Identity;
 			FTransform HipTransformCS = FTransform::Identity;
 			FTransform BallToFoot = FTransform::Identity;
