Foot Placement: Fix for pelvis placement when using weight rebalancing on slopes

diff --git a/Engine/Plugins/Animation/AnimationWarping/Source/Runtime/Private/BoneControllers/AnimNode_FootPlacement.cpp b/Engine/Plugins/Animation/AnimationWarping/Source/Runtime/Private/BoneControllers/AnimNode_FootPlacement.cpp
--- a/Engine/Plugins/Animation/AnimationWarping/Source/Runtime/Private/BoneControllers/AnimNode_FootPlacement.cpp
+++ b/Engine/Plugins/Animation/AnimationWarping/Source/Runtime/Private/BoneControllers/AnimNode_FootPlacement.cpp
@@ -1617,7 +1617,7 @@ FTransform FAnimNode_FootPlacement::SolvePelvis(const UE::Anim::FootPlacement::F
 		}
 
 		// Remove the vertical component
-		OffsetAverage -= OffsetAverage * Context.ApproachDirCS;
+		OffsetAverage += OffsetAverage * Context.ApproachDirCS;
 
 		PelvisOffsetDelta += OffsetAverage * PelvisSettings.HorizontalRebalancingWeight;
 	}
