Foot Placement:
-New setting to enable/disable pelvis/foot interpolation
-New modes for actor movement compensation behavior
-New setting to clamp how much feet clip into geometry
-New pelvis height modes to improve uphill/downhill behavior
-Better ordering/UX of node's settings
-Added debug vis for traces and pelvis target

diff --git a/Engine/Plugins/Animation/AnimationWarping/Source/Runtime/Private/BoneControllers/AnimNode_FootPlacement.cpp b/Engine/Plugins/Animation/AnimationWarping/Source/Runtime/Private/BoneControllers/AnimNode_FootPlacement.cpp
--- a/Engine/Plugins/Animation/AnimationWarping/Source/Runtime/Private/BoneControllers/AnimNode_FootPlacement.cpp
+++ b/Engine/Plugins/Animation/AnimationWarping/Source/Runtime/Private/BoneControllers/AnimNode_FootPlacement.cpp
@@ -14,6 +14,7 @@ DECLARE_CYCLE_STAT(TEXT("Foot Placement Eval"), STAT_FootPlacement_Eval, STATGRO
 static TAutoConsoleVariable<bool> CVarAnimNodeFootPlacementEnable(TEXT("a.AnimNode.FootPlacement.Enable"), true, TEXT("Enable/Disable Foot Placement"));
 static TAutoConsoleVariable<bool> CVarAnimNodeFootPlacementEnableLock(TEXT("a.AnimNode.FootPlacement.Enable.Lock"), true, TEXT("Enable/Disable Foot Locking"));
 static TAutoConsoleVariable<bool> CVarAnimNodeFootPlacementDebug(TEXT("a.AnimNode.FootPlacement.Debug"), false, TEXT("Turn on visualization debugging for Foot Placement"));
+static TAutoConsoleVariable<bool> CVarAnimNodeFootPlacementDebugTraces(TEXT("a.AnimNode.FootPlacement.Debug.Traces"), false, TEXT("Turn on visualization debugging for foot ground traces"));
 static TAutoConsoleVariable<int> CVarAnimNodeFootPlacementDebugDrawHistory(TEXT("a.AnimNode.FootPlacement.Debug.DrawHistory"), 0,
 	TEXT("Turn on history visualization debugging 0 = Disabled, -1 = Pelvis, >1 = Foot Index. Clear with FlushPersistentDebugLines"));
 #endif
@@ -165,6 +166,23 @@ namespace UE::Anim::FootPlacement
 		}
 
 		OutImpactNormalWS = HitResult.ImpactNormal;
+
+#if ENABLE_ANIM_DEBUG
+		if (CVarAnimNodeFootPlacementDebugTraces.GetValueOnAnyThread())
+		{
+			Context.CSPContext.AnimInstanceProxy->AnimDrawDebugPoint(
+				TraceStart, 10.0f, FColor::Purple, false, -1.0f, SDPG_Foreground);
+
+			Context.CSPContext.AnimInstanceProxy->AnimDrawDebugPoint(
+				OutImpactLocationWS, 10.0f, FColor::Purple, false, -1.0f, SDPG_Foreground);
+
+			Context.CSPContext.AnimInstanceProxy->AnimDrawDebugLine(
+				TraceStart, OutImpactLocationWS,
+				FColor::Purple,
+				false, -1.0f, 1.0f, SDPG_Foreground);
+		}
+#endif
+
 		return true;
 	}
 
@@ -489,51 +507,65 @@ void FAnimNode_FootPlacement::UpdatePlantingPlaneInterpolation(
 		// Trace against complex geometry only to plant accurately
 		UE::Anim::FootPlacement::FindPlantPlane(Context, TraceSettings, FootTransformWS.GetLocation(), true, InOutPlantPlane, ImpactLocationWS);
 	}
-	
-	const FVector CurrPlaneIntersection = UE::Anim::FootPlacement::PointDirectionPlaneIntersection(
-		FootTransformWS.GetLocation(),
-		TraceDirection,
-		InOutPlantPlane);
-
-	const FVector LastPlaneIntersection = UE::Anim::FootPlacement::PointDirectionPlaneIntersection(
-		LastAlignedFootTransform.GetLocation(),
-		TraceDirection,
-		LastPlantPlane);
-
-	const FVector PrevPlaneIntersection = UE::Anim::FootPlacement::PointDirectionPlaneIntersection(
-		FootTransformWS.GetLocation(),
-		TraceDirection,
-		LastPlantPlane);
-
-	const float LastPlaneDeltaZ = LastPlaneIntersection.Z - CurrPlaneIntersection.Z;
-	const float PrevPlaneDeltaZ = PrevPlaneIntersection.Z - CurrPlaneIntersection.Z;
-	const float AdjustedPrevZ = FMath::Abs(LastPlaneDeltaZ) < FMath::Abs(PrevPlaneDeltaZ) ?
-		LastPlaneIntersection.Z : PrevPlaneIntersection.Z;
-
-	//TODO: replace by Z? Do some math and interpolate Plane.W!
-	const FVector AdjustedPrevPlaneIntersection = FVector(CurrPlaneIntersection.X,
-		CurrPlaneIntersection.Y,
-		AdjustedPrevZ);
-
-	const FVector PlantPlaneSpringLocation = UKismetMathLibrary::VectorSpringInterp(
-		AdjustedPrevPlaneIntersection, CurrPlaneIntersection, InOutInterpData.GroundTranslationSpringState,
-		InterpolationSettings.FloorLinearStiffness,
-		InterpolationSettings.FloorLinearDamping,
-		Context.UpdateDeltaTime, 1.0f, 0.0f);
 
-	FQuat FloorNormalRotation = FQuat::FindBetweenNormals(
-		LastPlantPlane.GetNormal(), InOutPlantPlane.GetNormal());
-	const FQuat FloorSpringNormalRotation = UKismetMathLibrary::QuaternionSpringInterp(
-		FQuat::Identity, FloorNormalRotation, InOutInterpData.GroundRotationSpringState,
-		InterpolationSettings.FloorAngularStiffness,
-		InterpolationSettings.FloorAngularDamping,
-		Context.UpdateDeltaTime, 1.0f, 0.0f);
+	if (InterpolationSettings.bEnableFloorInterpolation)
+	{
+		const FVector CurrPlaneIntersection = UE::Anim::FootPlacement::PointDirectionPlaneIntersection(
+			FootTransformWS.GetLocation(),
+			TraceDirection,
+			InOutPlantPlane);
+
+		const FVector LastPlaneIntersection = UE::Anim::FootPlacement::PointDirectionPlaneIntersection(
+			LastAlignedFootTransform.GetLocation(),
+			TraceDirection,
+			LastPlantPlane);
+
+		const FVector PrevPlaneIntersection = UE::Anim::FootPlacement::PointDirectionPlaneIntersection(
+			FootTransformWS.GetLocation(),
+			TraceDirection,
+			LastPlantPlane);
+
+		const float LastPlaneDeltaZ = LastPlaneIntersection.Z - CurrPlaneIntersection.Z;
+		const float PrevPlaneDeltaZ = PrevPlaneIntersection.Z - CurrPlaneIntersection.Z;
+		const float AdjustedPrevZ = FMath::Abs(LastPlaneDeltaZ) < FMath::Abs(PrevPlaneDeltaZ) ?
+			LastPlaneIntersection.Z : PrevPlaneIntersection.Z;
+
+		//TODO: replace by Z? Do some math and interpolate Plane.W!
+		const FVector AdjustedPrevPlaneIntersection = FVector(CurrPlaneIntersection.X,
+			CurrPlaneIntersection.Y,
+			AdjustedPrevZ);
+
+		FVector PlantPlaneSpringLocation = UKismetMathLibrary::VectorSpringInterp(
+			AdjustedPrevPlaneIntersection, CurrPlaneIntersection, InOutInterpData.GroundTranslationSpringState,
+			InterpolationSettings.FloorLinearStiffness,
+			InterpolationSettings.FloorLinearDamping,
+			Context.UpdateDeltaTime, 1.0f, 0.0f);
+
+		if (TraceSettings.MaxGroundPenetration >= 0.0f)
+		{
+			// Prevent the foot from clipping too much into geometry due to interpolation
+			const float DistanceToGroundPlane =
+				UE::Anim::FootPlacement::GetDistanceToPlaneAlongDirection(PlantPlaneSpringLocation, InOutPlantPlane, Context.ApproachDirCS);
+			const float PenetrationAmount = -DistanceToGroundPlane - TraceSettings.MaxGroundPenetration;
+			if (PenetrationAmount > 0.0f)
+			{
+				PlantPlaneSpringLocation -= Context.ApproachDirCS * PenetrationAmount;
+			}
+		}
 
-	const FVector PlantPlaneSpringNormal = FloorSpringNormalRotation.RotateVector(LastPlantPlane.GetNormal());
+		FQuat FloorNormalRotation = FQuat::FindBetweenNormals(
+			LastPlantPlane.GetNormal(), InOutPlantPlane.GetNormal());
+		const FQuat FloorSpringNormalRotation = UKismetMathLibrary::QuaternionSpringInterp(
+			FQuat::Identity, FloorNormalRotation, InOutInterpData.GroundRotationSpringState,
+			InterpolationSettings.FloorAngularStiffness,
+			InterpolationSettings.FloorAngularDamping,
+			Context.UpdateDeltaTime, 1.0f, 0.0f);
 
-	const FPlane PlantingPlane = FPlane(PlantPlaneSpringLocation, PlantPlaneSpringNormal);
+		const FVector PlantPlaneSpringNormal = FloorSpringNormalRotation.RotateVector(LastPlantPlane.GetNormal());
+		const FPlane PlantingPlane = FPlane(PlantPlaneSpringLocation, PlantPlaneSpringNormal);
 
-	InOutPlantPlane = PlantingPlane;
+		InOutPlantPlane = PlantingPlane;
+	}
 }
 
 void FAnimNode_FootPlacement::DeterminePlantType(
@@ -551,7 +583,7 @@ void FAnimNode_FootPlacement::DeterminePlantType(
 	InOutPlantData.bWantsToPlant = WantsToPlant(Context, LegInputPose);
 	InOutPlantData.PlantType = EPlantType::Unplanted;
 
-	if (!InOutPlantData.bWantsToPlant || !InOutPlantData.bCanReachTarget)
+	if (!InOutPlantData.bWantsToPlant)
 	{
 		return;
 	}
@@ -706,9 +738,7 @@ UE::Anim::FootPlacement::FPlantResult FAnimNode_FootPlacement::FinalizeFootAlign
 		(LegData.InputPose.FootTransformCS.GetLocation() - LegData.InputPose.HipTransformCS.GetLocation()).GetSafeNormal();
 	const FVector TargetHipToFootDir =
 		(CorrectedFootTransformCS.GetLocation() - FinalHipTransformCS.GetLocation()).GetSafeNormal();
-
-	// Assume the pelvis adjustments let us reach the spot, unless we're too over-extended
-	LegData.Plant.bCanReachTarget = true;
+	
 
 	const FVector FKHipToFoot =
 		LegData.InputPose.FootTransformCS.GetLocation() - LegData.InputPose.HipTransformCS.GetLocation();
@@ -738,17 +768,18 @@ UE::Anim::FootPlacement::FPlantResult FAnimNode_FootPlacement::FinalizeFootAlign
 					LegData.Plant.bCanReachTarget = false;
 				}
 
-				const FTransform FKHipToLeg =
-					LegData.InputPose.FootTransformCS.GetRelativeTransform(PelvisData.InputPose.FKTransformCS);
-				const FTransform FKLegAtCurrentHipCS =
-					FKHipToLeg * PelvisTransformCS;
-				FVector IKToFK = (CorrectedFootTransformCS.GetLocation() - FKLegAtCurrentHipCS.GetLocation()).GetSafeNormal();
-
 				const bool bRecentlyUnplanted = !bIsPlanted && LegData.Plant.TimeSinceFullyUnaligned == 0.0f;
 				// Try to keep the tip on spot if we're unplanting
+				// Don't do this until we've reached the plant target once
 				// TODO: Make this configurable? 
-				if (bRecentlyUnplanted || bIsPlanted)
+				if ((bRecentlyUnplanted || bIsPlanted) && LegData.Plant.bCanReachTarget)
 				{
+					const FTransform FKHipToLeg =
+						LegData.InputPose.FootTransformCS.GetRelativeTransform(PelvisData.InputPose.FKTransformCS);
+					const FTransform FKLegAtCurrentHipCS =
+						FKHipToLeg * PelvisTransformCS;
+					FVector IKToFK = (CorrectedFootTransformCS.GetLocation() - FKLegAtCurrentHipCS.GetLocation()).GetSafeNormal();
+
 					// Scale this value by our FK transition alpha to not pop
 					const float MaxPullTowardsHip =  FMath::Min(LegData.Bones.FootLength, HyperExtensionRemaining) * LegData.InputPose.AlignmentAlpha;
 					HyperExtensionRemaining -= MaxPullTowardsHip;
@@ -786,6 +817,11 @@ UE::Anim::FootPlacement::FPlantResult FAnimNode_FootPlacement::FinalizeFootAlign
 					CorrectedFootTransformCS.SetLocation(NotHyperextendedPlantLocation);
 				}
 			}
+			else
+			{
+				// No overextension, therefore we can reach the target, and can do future tip/ball adjustments
+				LegData.Plant.bCanReachTarget = true;
+			}
 
 #if ENABLE_ANIM_DEBUG
 			DebugData.LegsInfo[LegData.Idx].HyperExtensionAmount = HyperExtensionAmount;
@@ -1057,17 +1093,21 @@ void FAnimNode_FootPlacement::EvaluateSkeletalControl_AnyThread(FComponentSpaceP
 
 	// Based on the ground alignment, search for the best Pelvis transform
 	FTransform PelvisTransformCS = SolvePelvis(FootPlacementContext);
-	PelvisTransformCS = UpdatePelvisInterpolation(FootPlacementContext, PelvisTransformCS);
-	check(!PelvisTransformCS.ContainsNaN());
-	OutBoneTransforms.Add(FBoneTransform(PelvisData.Bones.FkBoneIndex, PelvisTransformCS));
-
 #if ENABLE_ANIM_DEBUG
+	const FTransform PelvisTargetTransformCS = PelvisTransformCS;
 	if (CVarAnimNodeFootPlacementDebug.GetValueOnAnyThread())
 	{
+		// Debug text header
 		FString HeaderMessage = FString::Printf(TEXT("FOOT PLACEMENT DEBUG"));
 		FootPlacementContext.CSPContext.AnimInstanceProxy->AnimDrawDebugOnScreenMessage(HeaderMessage, FColor::Cyan);
 	}
 #endif
+	if (PelvisSettings.bEnableInterpolation)
+	{
+		PelvisTransformCS = UpdatePelvisInterpolation(FootPlacementContext, PelvisTransformCS);
+	}
+	check(!PelvisTransformCS.ContainsNaN());
+	OutBoneTransforms.Add(FBoneTransform(PelvisData.Bones.FkBoneIndex, PelvisTransformCS));
 
 	for (int32 FootIndex = 0; FootIndex < LegsData.Num(); ++FootIndex)
 	{
@@ -1130,12 +1170,17 @@ void FAnimNode_FootPlacement::EvaluateSkeletalControl_AnyThread(FComponentSpaceP
 
 			const FTransform PelvisTransformWS = PelvisTransformCS * ComponentTransform;
 			const FTransform BasePelvisTransformWS = PelvisData.InputPose.FKTransformCS * ComponentTransform;
+			const FTransform PelvisTargetTransformWS = PelvisTargetTransformCS * ComponentTransform;
 
 			AnimInstanceProxy->AnimDrawDebugPoint(
 				PelvisTransformWS.GetLocation(), 20.0f, FColor::Green, false, -1.0f, SDPG_Foreground);
 
 			AnimInstanceProxy->AnimDrawDebugPoint(
 				BasePelvisTransformWS.GetLocation(), 20.0f, FColor::Blue, false, -1.0f, SDPG_Foreground);
+
+			// Draw pelvis interpolation target
+			AnimInstanceProxy->AnimDrawDebugPoint(
+				PelvisTargetTransformWS.GetLocation(), 10.0f, FColor::Purple, false, -1.0f, SDPG_Foreground);
 		}
 	}
 #endif
@@ -1330,14 +1375,6 @@ void FAnimNode_FootPlacement::ProcessCharacterState(const UE::Anim::FootPlacemen
 
 	CharacterData.ComponentTransformWS = Context.OwningComponentToWorld;
 	const FVector ComponentLocationWS = CharacterData.ComponentTransformWS.GetLocation();
-	CharacterData.NumFKPlanted = 0;
-	for (const UE::Anim::FootPlacement::FLegRuntimeData& LegData : LegsData)
-	{
-		if (FMath::IsNearlyEqual(LegData.InputPose.AlignmentAlpha, 1.0f))
-		{
-			++CharacterData.NumFKPlanted;
-		}
-	}
 
 	const bool bWasOnGround = CharacterData.bIsOnGround;
 	CharacterData.bIsOnGround = !Context.MovementComponent ?
@@ -1346,20 +1383,30 @@ void FAnimNode_FootPlacement::ProcessCharacterState(const UE::Anim::FootPlacemen
 			(Context.MovementComponent->MovementMode == MOVE_NavWalking)) &&
 		Context.MovementComponent->CurrentFloor.bBlockingHit;
 
-	if (CharacterData.bIsOnGround && bWasOnGround && PelvisSettings.bCompensateForSuddenCapsuleMoves)
+	if (CharacterData.bIsOnGround && bWasOnGround && (PelvisSettings.ActorMovementCompensationMode != EActorMovementCompensationMode::ComponentSpace))
 	{
-		// Compensate for sudden capsule moves
-		const FVector CapsuleFloorNormalWS = Context.GetMovementComponentFloorNormal();
-		const FVector OwningComponentAdjustedLastLocationWS =
-			(FMath::Abs(Context.ApproachDirWS | CapsuleFloorNormalWS) > DELTA) ?
-			UE::Anim::FootPlacement::PointDirectionPlaneIntersection(
-				ComponentLocationWS,
-				Context.ApproachDirWS,
-				FPlane(LastComponentLocationWS, CapsuleFloorNormalWS)) :
-			ComponentLocationWS;
+		FVector OwningComponentAdjustedLastLocationWS;
+		if (PelvisSettings.ActorMovementCompensationMode == EActorMovementCompensationMode::SuddenMotionOnly)
+		{
+			// Compensate for sudden capsule moves
+			const FVector CapsuleFloorNormalWS = Context.GetMovementComponentFloorNormal();
+			OwningComponentAdjustedLastLocationWS =
+				(FMath::Abs(Context.ApproachDirWS | CapsuleFloorNormalWS) > DELTA) ?
+				UE::Anim::FootPlacement::PointDirectionPlaneIntersection(
+					ComponentLocationWS,
+					Context.ApproachDirWS,
+					FPlane(LastComponentLocationWS, CapsuleFloorNormalWS)) :
+				ComponentLocationWS;
+		}
+		else //if (PelvisSettings.ActorMovementCompensationMode == EActorMovementCompensationMode::WorldSpace)
+		{
+			// Compensate for all moves
+			OwningComponentAdjustedLastLocationWS = LastComponentLocationWS;
+		}
 
+		// Only compensate vertical motion
 		const FVector CapsuleMoveOffsetWS =
-			ComponentLocationWS - OwningComponentAdjustedLastLocationWS;
+			(ComponentLocationWS - LastComponentLocationWS) * -Context.ApproachDirWS;
 		if (!CapsuleMoveOffsetWS.IsNearlyZero(KINDA_SMALL_NUMBER))
 		{
 			const FVector CapsuleMoveOffsetCS =
@@ -1368,6 +1415,16 @@ void FAnimNode_FootPlacement::ProcessCharacterState(const UE::Anim::FootPlacemen
 			PelvisData.Interpolation.PelvisTranslationOffset -= CapsuleMoveOffsetCS;
 		}
 	}
+
+	{
+		// Use the component's move delta, instead of the movement component's velocity, since it doesn't account for uphill/downhill velocity
+		const FVector CapsuleMoveOffsetWS =
+			(ComponentLocationWS - LastComponentLocationWS);
+		const FVector CapsuleMoveOffsetCS =
+			Context.OwningComponentToWorld.InverseTransformVectorNoScale(CapsuleMoveOffsetWS);
+
+		CharacterData.ComponentVelocityCS = CapsuleMoveOffsetCS / CachedDeltaTime;
+	}
 }
 
 void FAnimNode_FootPlacement::ProcessFootAlignment(
@@ -1560,16 +1617,22 @@ FTransform FAnimNode_FootPlacement::SolvePelvis(const UE::Anim::FootPlacement::F
 {
 	using namespace UE::Anim::FootPlacement;
 
+	TBitArray<> RelevantFeet = FindRelevantFeet(Context);
+	const int32 RelevantFeetNum = RelevantFeet.CountSetBits();
+
 	// Taken from http://runevision.com/thesis/rune_skovbo_johansen_thesis.pdf
 	// Chapter 7.4.2
 
 	float MaxOffsetMin = BIG_NUMBER;
 	float DesiredOffsetMin = BIG_NUMBER;
-	float DesiredOffsetSum = 0.0f;
+	float DesiredOffsetAvg = 0.0f;
 	float MinOffsetMax = -BIG_NUMBER;
 
-	for (const FLegRuntimeData& LegData : LegsData)
+	for (TConstSetBitIterator It(RelevantFeet); It; ++It)
 	{
+		const int32 LegIndex = It.GetIndex();
+		const FLegRuntimeData& LegData = LegsData[LegIndex];
+
 		FPelvisOffsetRangeForLimb PelvisOffsetRangeCS;
 		FindPelvisOffsetRangeForLimb(
 			Context,
@@ -1583,12 +1646,11 @@ FTransform FAnimNode_FootPlacement::SolvePelvis(const UE::Anim::FootPlacement::F
 		const float MaxOffset = PelvisOffsetRangeCS.MaxExtension;
 		const float MinOffset = PelvisOffsetRangeCS.MinExtension;
 
-		DesiredOffsetSum += DesiredOffset;
+		DesiredOffsetAvg += DesiredOffset / RelevantFeetNum;
 		DesiredOffsetMin = FMath::Min(DesiredOffsetMin, DesiredOffset);
 		MaxOffsetMin = FMath::Min(MaxOffsetMin, MaxOffset);
 		MinOffsetMax = FMath::Max(MinOffsetMax, MinOffset);
 	}
-	const float DesiredOffsetAvg = DesiredOffsetSum / LegsData.Num();
 	const float MinToAvg = DesiredOffsetAvg - DesiredOffsetMin;
 	const float MinToMax = MaxOffsetMin - DesiredOffsetMin;
 
@@ -1602,7 +1664,7 @@ FTransform FAnimNode_FootPlacement::SolvePelvis(const UE::Anim::FootPlacement::F
 		DesiredOffsetMin + ((MinToAvg * MinToMax) / Divisor);
 
 	// Adjust the hips to prevent over-compression
-	PelvisOffsetZ = FMath::Clamp(PelvisOffsetZ, MinOffsetMax, MaxOffsetMin);
+	//PelvisOffsetZ = FMath::Clamp(PelvisOffsetZ, MinOffsetMax, MaxOffsetMin);
 
 	FVector PelvisOffsetDelta = -PelvisOffsetZ * Context.ApproachDirCS;
 
@@ -1628,6 +1690,84 @@ FTransform FAnimNode_FootPlacement::SolvePelvis(const UE::Anim::FootPlacement::F
 	return PelvisTransformCS;
 }
 
+TBitArray<> FAnimNode_FootPlacement::FindRelevantFeet(const UE::Anim::FootPlacement::FEvaluationContext& Context)
+{
+	TBitArray<> RelevantFeet;
+
+	if (PelvisSettings.PelvisHeightMode == EPelvisHeightMode::AllLegs)
+	{
+		RelevantFeet.Init(true, LegsData.Num());
+		return RelevantFeet;
+	}
+
+
+	if (PelvisSettings.PelvisHeightMode == EPelvisHeightMode::FrontPlantedFeetUphill_FrontFeetDownhill)
+	{
+		if (!CharacterData.ComponentVelocityCS.IsNearlyZero())
+		{
+			TBitArray<> FrontFeet;
+			FrontFeet.SetNumUninitialized(LegsData.Num());
+			const FVector MovementDirection = CharacterData.ComponentVelocityCS.GetUnsafeNormal();
+			for (int32 LegIndex = 0; LegIndex < LegsData.Num(); ++LegIndex)
+			{
+				const FVector FootLocation = LegsData[LegIndex].InputPose.FootTransformCS.GetLocation();
+				const FVector MidpointToFoot = LegsData[LegIndex].InputPose.FootTransformCS.GetLocation() - PelvisData.InputPose.FootMidpointCS;
+				FrontFeet[LegIndex] = MidpointToFoot.Dot(MovementDirection) > 0.0f;
+			}
+
+			const int32 NumFrontFeet = FrontFeet.CountSetBits();
+			const bool bMovingDownHill = CharacterData.ComponentVelocityCS.Dot(Context.ApproachDirCS) > 0.0f;
+			if (bMovingDownHill)
+			{
+				// Consider all legs relevant when moving downhill. When using multiple feet, we tend to favour the lowest foot anyway.
+				RelevantFeet.Init(true, LegsData.Num());
+			}
+			else
+			{
+				TBitArray<> PlantedFeet = FindPlantedFeet(Context);
+				// Prefer the front-most planted foot, if any
+				RelevantFeet = TBitArray<>::BitwiseAND(PlantedFeet, FrontFeet , EBitwiseOperatorFlags::MinSize);
+
+				const int32 NumRelevantFeet = RelevantFeet.CountSetBits();
+				if (NumRelevantFeet == 0)
+				{
+					// No front-most planted feet. Fallback to any planted feet
+					RelevantFeet = PlantedFeet;
+				}
+			}
+		}
+	}
+	else //if (PelvisSettings.PelvisHeightMode == EPelvisHeightMode::AllPlantedFeet)
+	{
+		TBitArray<> PlantedFeet = FindPlantedFeet(Context);
+		RelevantFeet = PlantedFeet;
+	}
+	
+	if (RelevantFeet.CountSetBits() == 0)
+	{
+		// If we reach this point with no planted feet, fall-back to all feet
+		RelevantFeet.Init(true, LegsData.Num());
+	}
+
+	return RelevantFeet;
+}
+
+TBitArray<> FAnimNode_FootPlacement::FindPlantedFeet(const UE::Anim::FootPlacement::FEvaluationContext& Context)
+{
+	TBitArray<> PlantedFeet;
+	PlantedFeet.Init(true, LegsData.Num());
+	for (int32 LegIndex = 0; LegIndex < LegsData.Num(); ++LegIndex)
+	{
+		// Using alignment alpha, instead of bWantsToPlant allows us to start the weight-shift a few frames before planting
+		if (LegsData[LegIndex].InputPose.AlignmentAlpha == 0.0f)
+		{
+			PlantedFeet[LegIndex] = false;
+		}
+	}
+
+	return PlantedFeet;
+}
+
 FTransform FAnimNode_FootPlacement::UpdatePelvisInterpolation(
 	const UE::Anim::FootPlacement::FEvaluationContext& Context,
 	const FTransform& TargetPelvisTransform)
@@ -1648,11 +1788,12 @@ FTransform FAnimNode_FootPlacement::UpdatePelvisInterpolation(
 	}
 
 	// Spring interpolation may cause hyperextension/compression so we solve that in FinalizeFootAlignment
-	PelvisData.Interpolation.PelvisTranslationOffset = UKismetMathLibrary::VectorSpringInterp(
+	const FVector NewTranslationOffset = UKismetMathLibrary::VectorSpringInterp(
 		PelvisData.Interpolation.PelvisTranslationOffset, DesiredPelvisOffset, PelvisData.Interpolation.PelvisTranslationSpringState,
 		PelvisSettings.LinearStiffness,
 		PelvisSettings.LinearDamping,
 		Context.UpdateDeltaTime, 1.0f, 0.0f);
+	PelvisData.Interpolation.PelvisTranslationOffset = NewTranslationOffset;
 
 	OutPelvisTransform.SetLocation(
 		PelvisData.InputPose.FKTransformCS.GetLocation() + PelvisData.Interpolation.PelvisTranslationOffset);
diff --git a/Engine/Plugins/Animation/AnimationWarping/Source/Runtime/Public/BoneControllers/AnimNode_FootPlacement.h b/Engine/Plugins/Animation/AnimationWarping/Source/Runtime/Public/BoneControllers/AnimNode_FootPlacement.h
--- a/Engine/Plugins/Animation/AnimationWarping/Source/Runtime/Public/BoneControllers/AnimNode_FootPlacement.h
+++ b/Engine/Plugins/Animation/AnimationWarping/Source/Runtime/Public/BoneControllers/AnimNode_FootPlacement.h
@@ -118,7 +118,7 @@ namespace UE::Anim::FootPlacement
 		/* Interpolation */
 		struct FInterpolationData
 		{
-			FVector PelvisTranslationOffset;
+			FVector PelvisTranslationOffset = FVector::ZeroVector;
 			FVectorSpringState PelvisTranslationSpringState;
 		} Interpolation;
 	};
@@ -126,7 +126,7 @@ namespace UE::Anim::FootPlacement
 	struct FCharacterData
 	{
 		FTransform ComponentTransformWS = FTransform::Identity;
-		int NumFKPlanted = 0;
+		FVector ComponentVelocityCS = FVector::ZeroVector;
 		bool bIsOnGround = false;
 	};
 
@@ -195,17 +195,20 @@ public:
 	UPROPERTY(EditAnywhere, Category = "Plant Settings")
 	float UnplantAngularDamping = 1.0f;
 
-	UPROPERTY(EditAnywhere, Category = "Plant Settings")
+	UPROPERTY(EditAnywhere, Category = "Plant Settings", meta=(EditCondition="bEnableFloorInterpolation", DisplayAfter="bEnableFloorInterpolation"))
 	float FloorLinearStiffness = 1000.0f;
 
-	UPROPERTY(EditAnywhere, Category = "Plant Settings")
+	UPROPERTY(EditAnywhere, Category = "Plant Settings", meta=(EditCondition="bEnableFloorInterpolation", DisplayAfter="bEnableFloorInterpolation"))
 	float FloorLinearDamping = 1.0f;
 
-	UPROPERTY(EditAnywhere, Category = "Plant Settings")
+	UPROPERTY(EditAnywhere, Category = "Plant Settings", meta=(EditCondition="bEnableFloorInterpolation", DisplayAfter="bEnableFloorInterpolation"))
 	float FloorAngularStiffness = 450.0f;
 
-	UPROPERTY(EditAnywhere, Category = "Plant Settings")
+	UPROPERTY(EditAnywhere, Category = "Plant Settings", meta=(EditCondition="bEnableFloorInterpolation", DisplayAfter="bEnableFloorInterpolation"))
 	float FloorAngularDamping = 1.0f;
+
+	UPROPERTY(EditAnywhere, Category = "Plant Settings")
+	bool bEnableFloorInterpolation = true;
 };
 
 USTRUCT(BlueprintType)
@@ -214,10 +217,6 @@ struct FFootPlacementTraceSettings
 	GENERATED_BODY()
 
 public:
-	// TODO: implement?
-	// Is tracing enabled?
-	UPROPERTY(EditAnywhere, Category = "Trace Settings")
-	bool bEnabled = true;
 
 	// A negative value extends the trace length above the bone
 	UPROPERTY(EditAnywhere, Category = "Trace Settings", meta=(EditCondition="bEnabled", DisplayAfter="bEnabled"))
@@ -236,6 +235,11 @@ public:
 	UPROPERTY(EditAnywhere, Category = "Trace Settings", meta=(EditCondition="bEnabled", DisplayAfter="bEnabled"))
 	TEnumAsByte<ETraceTypeQuery> ComplexTraceChannel = TraceTypeQuery1;
 
+	// How much the feet can penetrate the ground geometry. It's recommended to allow some to account for interpolation
+	// Negative values disable this effect
+	UPROPERTY(EditAnywhere, Category = "Trace Settings")
+	float MaxGroundPenetration = 10.0f;
+
 	// How much we align to simple vs complex collision when the foot is in flight
 	// Tracing against simple geometry (i.e. it's common for stairs to have simplified ramp collisions) can provide a 
 	// smoother trajectory when the foot is in flight
@@ -245,6 +249,11 @@ public:
 	// The channel to use for our simple trace
 	UPROPERTY(EditAnywhere, Category = "Trace Settings", meta = (EditCondition = "bEnabled", DisplayAfter = "bEnabled"))
 	TEnumAsByte<ETraceTypeQuery> SimpleTraceChannel = TraceTypeQuery1;
+
+	// Enabling tracing for ground alignment
+	// @TODO: Use ground normal when not tracing
+	UPROPERTY(EditAnywhere, Category = "Trace Settings")
+	bool bEnabled = true;
 };
 
 USTRUCT()
@@ -263,6 +272,32 @@ public:
 	void Initialize(const FAnimationInitializeContext& Context);
 };
 
+UENUM(BlueprintType)
+enum class EPelvisHeightMode : uint8
+{
+	// Consider all legs for pelvis height, whether they're planted or not
+	// Generally good for flat/not too steep ground
+	AllLegs,
+	// Consider only the planted feet when calculating pelvis height
+	// Generally good we pelvis height to be defined by the weight supporting leg
+	AllPlantedFeet,
+	// When moving uphill, use the front foot, as long as it's planted.
+	// It's recommended to also enable pelvis interpolation to smooth out the swap between what's considered the "planted" leg
+	// When moving downhill, both feet will be considered relevant.
+	// The algorithm tends to prefer the lower foot, except when the higher foot would become over-compresseed.
+	FrontPlantedFeetUphill_FrontFeetDownhill,
+};
+
+UENUM(BlueprintType)
+enum class EActorMovementCompensationMode : uint8
+{
+	// Keep pelvis component-space and follow along all of the actor's vertical ground movement
+	ComponentSpace,
+	// Hold pelvis world-space and ignore the actor's vertical ground movement. Let springs interpolate the difference
+	WorldSpace,
+	// Keep pelvis component-space, but hold world-space transform when the actor does sudden changes (i.e. a big step), and let springs interpolate the difference.
+	SuddenMotionOnly
+};
 
 USTRUCT(BlueprintType)
 struct FFootPlacementPelvisSettings
@@ -275,28 +310,34 @@ public:
 	UPROPERTY(EditAnywhere, Category = "Pelvis Settings")
 	float MaxOffset = 50.0f;
 
-	UPROPERTY(EditAnywhere, Category = "Pelvis Settings")
+	UPROPERTY(EditAnywhere, Category = "Pelvis Settings", meta=(EditCondition="bEnableInterpolation", DisplayAfter="bEnableInterpolation"))
 	float LinearStiffness = 350.0f;
 
-	UPROPERTY(EditAnywhere, Category = "Pelvis Settings")
+	UPROPERTY(EditAnywhere, Category = "Pelvis Settings", meta=(EditCondition="bEnableInterpolation", DisplayAfter="bEnableInterpolation"))
 	float LinearDamping = 1.0f;
 
 	// How much we move the pelvis horizontally to re-balance the characters weight due to foot offsets.
 	// A value of 0 will disable this effect.
 	// Higher values may move the mesh outside of the character's capsule
-	UPROPERTY(EditAnywhere, Category = "Pelvis Settings", meta = (ClampMin = "0.0", UIMin = "0.0", ClampMax = "1.0", UIMax = "1.0"))
+	UPROPERTY(EditAnywhere, Category = "Pelvis Settings", meta = (ClampMin = "0.0", UIMin = "0.0", ClampMax = "1.0", UIMax = "1.0", DisplayAfter="bEnableInterpolation"))
 	float HorizontalRebalancingWeight = 0.3f;
 
 	// Max horizontal foot adjustment we consider to lower the hips
 	// This can be used to prevent the hips from dropping too low when the feet are locked
 	// Exceeding this value will first attempt to roll the planted feet, and then slide
-	UPROPERTY(EditAnywhere, Category = "Pelvis Settings")
+	UPROPERTY(EditAnywhere, Category = "Pelvis Settings", meta=(DisplayAfter="bEnableInterpolation"))
 	float MaxOffsetHorizontal = 10.0f;
 
-	// This is used to hold the Pelvis's interpolator in a fixed spot when the capsule suddenly moves (i.e. on a big step)
-	// If your camera is directly attached to the character with little to no smoothing, you may want this disabled
+	UPROPERTY(EditAnywhere, Category = "Pelvis Settings", meta=(DisplayAfter="bEnableInterpolation"))
+	EPelvisHeightMode PelvisHeightMode = EPelvisHeightMode::AllLegs;
+
+	// This is used to hold the Pelvis's interpolator in a fixed spot when the capsule moves up/down
+	// If your camera is directly attached to the character with little to no smoothing, you may want this on ComponentSpace
+	UPROPERTY(EditAnywhere, Category = "Pelvis Settings", meta=(DisplayAfter="bEnableInterpolation"))
+	EActorMovementCompensationMode ActorMovementCompensationMode = EActorMovementCompensationMode::SuddenMotionOnly;
+
 	UPROPERTY(EditAnywhere, Category = "Pelvis Settings")
-	bool bCompensateForSuddenCapsuleMoves = true;
+	bool bEnableInterpolation = true;
 };
 
 USTRUCT()
@@ -469,6 +510,8 @@ private:
 	// Evaluates the Pelvis offsets required by the different planting bones and produces a result that best
 	// acommodates all of them
 	FTransform SolvePelvis(const UE::Anim::FootPlacement::FEvaluationContext& Context);
+	TBitArray<> FindRelevantFeet(const UE::Anim::FootPlacement::FEvaluationContext& Context);
+	TBitArray<> FindPlantedFeet(const UE::Anim::FootPlacement::FEvaluationContext& Context);
 
 	FTransform UpdatePelvisInterpolation(
 		const UE::Anim::FootPlacement::FEvaluationContext& Context,
