Offset Root Bone
-Move offset root bones out of struct. Structs are behaving wrong w/ folding + it makes the pinning of the not-too related settings harder.
-Add showdebug vis of offset root bone
-Fix near-zero half life values
-Properly reset on relevancy change

diff --git a/Engine/Plugins/Animation/AnimationWarping/Source/Runtime/Private/BoneControllers/AnimNode_OffsetRootBone.cpp b/Engine/Plugins/Animation/AnimationWarping/Source/Runtime/Private/BoneControllers/AnimNode_OffsetRootBone.cpp
--- a/Engine/Plugins/Animation/AnimationWarping/Source/Runtime/Private/BoneControllers/AnimNode_OffsetRootBone.cpp
+++ b/Engine/Plugins/Animation/AnimationWarping/Source/Runtime/Private/BoneControllers/AnimNode_OffsetRootBone.cpp
@@ -19,28 +19,49 @@ namespace UE::Anim::OffsetRootBone
 	// Taken from https://theorangeduck.com/page/spring-roll-call#implicitspringdamper
 	static float DamperImplicit(float Halflife, float DeltaTime, float Epsilon = 1e-8f)
 	{
-		return (1.0f - FMath::InvExpApprox((0.69314718056f * DeltaTime) / (Halflife + Epsilon)));
+		// Halflife values very close to 0 approach infinity, and result in big motion spikes when Halflife < DeltaTime
+		// This is a hack, and adds some degree of frame-rate dependency, but it holds up even at lower frame-rates.
+		Halflife = FMath::Max(Halflife, DeltaTime);
+		return FMath::Clamp((1.0f - FMath::InvExpApprox((0.69314718056f * DeltaTime) / (Halflife + Epsilon))), 0.0f, 1.0f);
 	}
 }
 
-// TODO: implement 
 void FAnimNode_OffsetRootBone::GatherDebugData(FNodeDebugData& DebugData)
 {
+	Super::GatherDebugData(DebugData);
+
+	FString DebugLine = DebugData.GetNodeName(this);
+#if ENABLE_ANIM_DEBUG
+	if (const UEnum* ModeEnum = StaticEnum<EOffsetRootBoneMode>())
+	{
+		DebugLine += FString::Printf(TEXT("\n - Translation Mode: (%s)"), *(ModeEnum->GetNameStringByIndex(static_cast<int32>(GetTranslationMode()))));
+		DebugLine += FString::Printf(TEXT("\n - Rotation Mode: (%s)"), *(ModeEnum->GetNameStringByIndex(static_cast<int32>(GetRotationMode()))));
+	}
+	DebugLine += FString::Printf(TEXT("\n - Translation Halflife: (%.3fd)"), GetTranslationHalflife());
+	DebugLine += FString::Printf(TEXT("\n - Rotation Halflife: (%.3fd)"), GetRotationHalfLife());
+#endif
+	DebugData.AddDebugItem(DebugLine);
+	ComponentPose.GatherDebugData(DebugData);
 }
 
 void FAnimNode_OffsetRootBone::Initialize_AnyThread(const FAnimationInitializeContext& Context)
 {
 	FAnimNode_SkeletalControlBase::Initialize_AnyThread(Context);
 	AnimInstanceProxy = Context.AnimInstanceProxy;
-	ComponentTransform = Context.AnimInstanceProxy->GetComponentTransform();
-	SimulatedTranslation = ComponentTransform.GetLocation();
-	SimulatedRotation = ComponentTransform.GetRotation();
+	Reset(Context);
 }
 
 void FAnimNode_OffsetRootBone::UpdateInternal(const FAnimationUpdateContext& Context)
 {
 	FAnimNode_SkeletalControlBase::UpdateInternal(Context);
 	CachedDeltaTime = Context.GetDeltaTime();
+
+	// If we just became relevant and haven't been initialized yet, then reset.
+	if (!bIsFirstUpdate && UpdateCounter.HasEverBeenUpdated() && !UpdateCounter.WasSynchronizedCounter(Context.AnimInstanceProxy->GetUpdateCounter()))
+	{
+		Reset(Context);
+	}
+	UpdateCounter.SynchronizeWith(Context.AnimInstanceProxy->GetUpdateCounter());
 }
 
 void FAnimNode_OffsetRootBone::EvaluateSkeletalControl_AnyThread(FComponentSpacePoseContext& Output, TArray<FBoneTransform>& OutBoneTransforms)
@@ -63,10 +84,10 @@ void FAnimNode_OffsetRootBone::EvaluateSkeletalControl_AnyThread(FComponentSpace
 	const FTransform LastComponentTransform = ComponentTransform;
 	ComponentTransform = AnimInstanceProxy->GetComponentTransform();
 
-	const bool bShouldConsumeTranslationOffset = GetOffsetSettings().TranslationMode == EOffsetRootBoneMode::Accumulate ||
-										GetOffsetSettings().TranslationMode == EOffsetRootBoneMode::Interpolate;
-	const bool bShouldConsumeRotationOffset = GetOffsetSettings().RotationMode == EOffsetRootBoneMode::Accumulate ||
-										GetOffsetSettings().RotationMode == EOffsetRootBoneMode::Interpolate;
+	const bool bShouldConsumeTranslationOffset = GetTranslationMode() == EOffsetRootBoneMode::Accumulate ||
+										GetTranslationMode() == EOffsetRootBoneMode::Interpolate;
+	const bool bShouldConsumeRotationOffset = GetRotationMode() == EOffsetRootBoneMode::Accumulate ||
+										GetRotationMode() == EOffsetRootBoneMode::Interpolate;
 
 	FTransform RootMotionTransformDelta = FTransform::Identity;
 	if (bGraphDriven)
@@ -126,18 +147,18 @@ void FAnimNode_OffsetRootBone::EvaluateSkeletalControl_AnyThread(FComponentSpace
 	bModifyBone = CVarAnimNodeOffsetRootBoneModifyBone.GetValueOnAnyThread() == 1;
 #endif
 
-	if (GetOffsetSettings().TranslationMode == EOffsetRootBoneMode::Release ||
-		GetOffsetSettings().TranslationMode == EOffsetRootBoneMode::Interpolate)
+	if (GetTranslationMode() == EOffsetRootBoneMode::Release ||
+		GetTranslationMode() == EOffsetRootBoneMode::Interpolate)
 	{
 		FVector TranslationOffset = ComponentTransform.GetLocation() - SimulatedTranslation;
 
-		const float DampenAlpha = UE::Anim::OffsetRootBone::DamperImplicit(GetOffsetSettings().TranslationHalflife, CachedDeltaTime);
+		const float DampenAlpha = UE::Anim::OffsetRootBone::DamperImplicit(GetTranslationHalflife(), CachedDeltaTime);
 		FVector TranslationOffsetDelta = FMath::Lerp(FVector::ZeroVector, TranslationOffset, DampenAlpha);
 
-		if (GetOffsetSettings().bClampToTranslationVelocity)
+		if (GetClampToTranslationVelocity())
 		{
 			const float RootMotionDelta = RootMotionTransformDelta.GetLocation().Size();
-			const float MaxDelta = 	GetOffsetSettings().TranslationSpeedRatio * RootMotionDelta;
+			const float MaxDelta = 	GetTranslationSpeedRatio() * RootMotionDelta;
 
 			const float AdjustmentDelta = TranslationOffsetDelta.Size();
 			if (AdjustmentDelta > MaxDelta)
@@ -149,23 +170,23 @@ void FAnimNode_OffsetRootBone::EvaluateSkeletalControl_AnyThread(FComponentSpace
 		SimulatedTranslation = SimulatedTranslation + TranslationOffsetDelta;
 	}
 
-	if (GetOffsetSettings().RotationMode == EOffsetRootBoneMode::Release ||
-		GetOffsetSettings().RotationMode == EOffsetRootBoneMode::Interpolate)
+	if (GetRotationMode() == EOffsetRootBoneMode::Release ||
+		GetRotationMode() == EOffsetRootBoneMode::Interpolate)
 	{
 		FQuat RotationOffset = ComponentTransform.GetRotation() * SimulatedRotation.Inverse();
 		RotationOffset.Normalize();
 		RotationOffset = RotationOffset.W < 0.0 ? -RotationOffset : RotationOffset;
 
-		const float DampenAlpha = UE::Anim::OffsetRootBone::DamperImplicit(GetOffsetSettings().RotationHalfLife, CachedDeltaTime);
+		const float DampenAlpha = UE::Anim::OffsetRootBone::DamperImplicit(GetRotationHalfLife(), CachedDeltaTime);
 		FQuat RotationOffsetDelta = FQuat::Slerp(FQuat::Identity, RotationOffset, DampenAlpha);
 
-		if (GetOffsetSettings().bClampToRotationVelocity)
+		if (GetClampToRotationVelocity())
 		{
 			float RotationMotionAngleDelta;
 			FVector RootMotionRotationAxis;
 			RootMotionTransformDelta.GetRotation().ToAxisAndAngle(RootMotionRotationAxis, RotationMotionAngleDelta);
 
-			float MaxRotationAngle = GetOffsetSettings().RotationSpeedRatio * RotationMotionAngleDelta;
+			float MaxRotationAngle = GetRotationSpeedRatio() * RotationMotionAngleDelta;
 
 			FVector DeltaAxis;
 			float DeltaAngle;
@@ -180,19 +201,19 @@ void FAnimNode_OffsetRootBone::EvaluateSkeletalControl_AnyThread(FComponentSpace
 		SimulatedRotation = RotationOffsetDelta * SimulatedRotation;
 	}
 
-	if (GetOffsetSettings().MaxTranslationError >= 0.0f)
+	if (GetMaxTranslationError() >= 0.0f)
 	{
 		FVector TranslationOffset = ComponentTransform.GetLocation() - SimulatedTranslation;
 		const float TranslationOffsetSize = TranslationOffset.Size();
-		if (TranslationOffsetSize > GetOffsetSettings().MaxTranslationError)
+		if (TranslationOffsetSize > GetMaxTranslationError())
 		{
-			TranslationOffset = TranslationOffset.GetClampedToMaxSize(GetOffsetSettings().MaxTranslationError);
+			TranslationOffset = TranslationOffset.GetClampedToMaxSize(GetMaxTranslationError());
 			SimulatedTranslation = ComponentTransform.GetLocation() - TranslationOffset;
 		}
 	}
 
-	const float MaxAngleRadians = FMath::DegreesToRadians(GetOffsetSettings().MaxRotationError);
-	if (GetOffsetSettings().MaxRotationError >= 0.0f)
+	const float MaxAngleRadians = FMath::DegreesToRadians(GetMaxRotationError());
+	if (GetMaxRotationError() >= 0.0f)
 	{
 		FQuat RotationOffset = ComponentTransform.GetRotation().Inverse() * SimulatedRotation;
 		RotationOffset.Normalize();
@@ -227,15 +248,15 @@ void FAnimNode_OffsetRootBone::EvaluateSkeletalControl_AnyThread(FComponentSpace
 	if (bDebugging)
 	{
 		const float InnerCircleRadius = 40.0f;
-		const float CircleThickness = 0.5f;
-		const float ConeThickness = 0.2f;
+		const float CircleThickness = 1.5f;
+		const float ConeThickness = 0.3f;
 
 		const FTransform TargetBoneInitialTransformWorld = InputBoneTransform * ComponentTransform;
 		const FTransform TargetBoneTransformWorld = TargetBoneTransform * ComponentTransform;
 
-		if (GetOffsetSettings().MaxTranslationError >= 0.0f)
+		if (GetMaxTranslationError() >= 0.0f)
 		{
-			const float OuterCircleRadius = GetOffsetSettings().MaxTranslationError + InnerCircleRadius;
+			const float OuterCircleRadius = GetMaxTranslationError() + InnerCircleRadius;
 			AnimInstanceProxy->AnimDrawDebugCircle(ComponentTransform.GetLocation(), OuterCircleRadius, 36, FColor::Red,
 				FVector::UpVector, false, -1.0f, SDPG_World, CircleThickness);
 		}
@@ -246,7 +267,7 @@ void FAnimNode_OffsetRootBone::EvaluateSkeletalControl_AnyThread(FComponentSpace
 		AnimInstanceProxy->AnimDrawDebugCircle(TargetBoneTransformWorld.GetLocation(), InnerCircleRadius, 36, FColor::Green,
 			FVector::UpVector, false, -1.0f, SDPG_World, CircleThickness);
 
-		const int32 ConeSegments =  FMath::RoundUpToPowerOfTwo((GetOffsetSettings().MaxRotationError / 180.0f) * 12 );
+		const int32 ConeSegments =  FMath::RoundUpToPowerOfTwo((GetMaxRotationError() / 180.0f) * 12 );
 		const FVector ArcDirection = ComponentTransform.GetRotation().GetRightVector();
 		AnimInstanceProxy->AnimDrawDebugCone(TargetBoneTransformWorld.GetLocation(), 0.9f * InnerCircleRadius, ArcDirection,
 			MaxAngleRadians, 0.0f, ConeSegments, FColor::Red, false, -1.0f, SDPG_World, ConeThickness);
@@ -271,6 +292,8 @@ void FAnimNode_OffsetRootBone::EvaluateSkeletalControl_AnyThread(FComponentSpace
 		const bool bRootMotionOverridden = RootMotionProvider->OverrideRootMotion(RootMotionTransformDelta, Output.CustomAttributes);
 		ensure(bRootMotionOverridden);
 	}
+
+	bIsFirstUpdate = false;
 }
 
 bool FAnimNode_OffsetRootBone::IsValidToEvaluate(const USkeleton* Skeleton, const FBoneContainer& RequiredBones)
@@ -290,11 +313,6 @@ EWarpingEvaluationMode FAnimNode_OffsetRootBone::GetEvaluationMode() const
 	return GET_ANIM_NODE_DATA(EWarpingEvaluationMode, EvaluationMode);
 }
 
-const FOffsetRootBoneSettings& FAnimNode_OffsetRootBone::GetOffsetSettings() const
-{
-	return GET_ANIM_NODE_DATA(FOffsetRootBoneSettings, OffsetSettings);
-}
-
 const FVector& FAnimNode_OffsetRootBone::GetTranslationDelta() const
 {
 	return GET_ANIM_NODE_DATA(FVector, TranslationDelta);
@@ -304,3 +322,61 @@ const FRotator& FAnimNode_OffsetRootBone::GetRotationDelta() const
 {
 	return GET_ANIM_NODE_DATA(FRotator, RotationDelta);
 }
+
+EOffsetRootBoneMode FAnimNode_OffsetRootBone::GetTranslationMode() const
+{
+	return GET_ANIM_NODE_DATA(EOffsetRootBoneMode, TranslationMode);
+}
+
+EOffsetRootBoneMode FAnimNode_OffsetRootBone::GetRotationMode() const
+{
+	return GET_ANIM_NODE_DATA(EOffsetRootBoneMode, RotationMode);
+}
+
+float FAnimNode_OffsetRootBone::GetTranslationHalflife() const
+{
+	return GET_ANIM_NODE_DATA(float, TranslationHalflife);
+}
+
+float FAnimNode_OffsetRootBone::GetRotationHalfLife() const
+{
+	return GET_ANIM_NODE_DATA(float, RotationHalfLife);
+}
+
+float FAnimNode_OffsetRootBone::GetMaxTranslationError() const
+{
+	return GET_ANIM_NODE_DATA(float, MaxTranslationError);
+}
+
+float FAnimNode_OffsetRootBone::GetMaxRotationError() const
+{
+	return GET_ANIM_NODE_DATA(float, MaxRotationError);
+}
+
+bool FAnimNode_OffsetRootBone::GetClampToTranslationVelocity() const
+{
+	return GET_ANIM_NODE_DATA(bool, bClampToTranslationVelocity);
+}
+
+bool FAnimNode_OffsetRootBone::GetClampToRotationVelocity() const
+{
+	return GET_ANIM_NODE_DATA(bool, bClampToRotationVelocity);
+}
+
+float FAnimNode_OffsetRootBone::GetTranslationSpeedRatio() const
+{
+	return GET_ANIM_NODE_DATA(float, TranslationSpeedRatio);
+}
+
+float FAnimNode_OffsetRootBone::GetRotationSpeedRatio() const
+{
+	return GET_ANIM_NODE_DATA(float, RotationSpeedRatio);
+}
+
+void FAnimNode_OffsetRootBone::Reset(const FAnimationBaseContext& Context)
+{
+	ComponentTransform = Context.AnimInstanceProxy->GetComponentTransform();
+	SimulatedTranslation = ComponentTransform.GetLocation();
+	SimulatedRotation = ComponentTransform.GetRotation();
+	bIsFirstUpdate = true;
+}
diff --git a/Engine/Plugins/Animation/AnimationWarping/Source/Runtime/Public/BoneControllers/AnimNode_OffsetRootBone.h b/Engine/Plugins/Animation/AnimationWarping/Source/Runtime/Public/BoneControllers/AnimNode_OffsetRootBone.h
--- a/Engine/Plugins/Animation/AnimationWarping/Source/Runtime/Public/BoneControllers/AnimNode_OffsetRootBone.h
+++ b/Engine/Plugins/Animation/AnimationWarping/Source/Runtime/Public/BoneControllers/AnimNode_OffsetRootBone.h
@@ -27,71 +27,62 @@ enum class EOffsetRootBoneMode : uint8
 	Release,
 };
 
-USTRUCT(BlueprintType)
-struct FOffsetRootBoneSettings
+USTRUCT(BlueprintInternalUseOnly, Experimental)
+struct ANIMATIONWARPINGRUNTIME_API FAnimNode_OffsetRootBone : public FAnimNode_SkeletalControlBase
 {
-	GENERATED_BODY()
+	GENERATED_BODY();
+
+#if WITH_EDITORONLY_DATA
+	UPROPERTY(EditAnywhere, Category = Evaluation, meta=(FoldProperty))
+	EWarpingEvaluationMode EvaluationMode = EWarpingEvaluationMode::Graph;
 
 	// The translation offset behavior mode
-	UPROPERTY(EditAnywhere, BlueprintReadWrite, Category = Settings)
+	UPROPERTY(EditAnywhere, Category = Settings, meta = (FoldProperty, PinHiddenByDefault))
 	EOffsetRootBoneMode TranslationMode = EOffsetRootBoneMode::Interpolate;
 
 	// The rotation offset behavior mode
-	UPROPERTY(EditAnywhere, BlueprintReadWrite, Category = Settings)
+	UPROPERTY(EditAnywhere, Category = Settings, meta = (FoldProperty, PinHiddenByDefault))
 	EOffsetRootBoneMode RotationMode = EOffsetRootBoneMode::Interpolate;
 
 	// Controls how fast the translation offset is blended out
 	// Values closer to 0 make it faster
-	UPROPERTY(EditAnywhere, BlueprintReadWrite, Category = Settings)
+	UPROPERTY(EditAnywhere, Category = Settings, meta = (FoldProperty, PinHiddenByDefault))
 	float TranslationHalflife = 0.1f;
 
 	// Controls how fast the rotation offset is blended out
 	// Values closer to 0 make it faster
-	UPROPERTY(EditAnywhere, BlueprintReadWrite, Category = Settings)
+	UPROPERTY(EditAnywhere, Category = Settings, meta = (FoldProperty, PinHiddenByDefault))
 	float RotationHalfLife = 0.2f;
 
 	// How much the offset can deviate from the mesh component's translation in units
 	// Values lower than 0 disable this limit
-	UPROPERTY(EditAnywhere, BlueprintReadWrite, Category = Settings)
+	UPROPERTY(EditAnywhere, Category = Settings, meta = (FoldProperty, PinHiddenByDefault))
 	float MaxTranslationError = -1.0f;
 
 	// How much the offset can deviate from the mesh component's rotation in degrees
 	// Values lower than 0 disable this limit
-	UPROPERTY(EditAnywhere, BlueprintReadWrite, Category = Settings)
+	UPROPERTY(EditAnywhere, Category = Settings, meta = (FoldProperty, PinHiddenByDefault))
 	float MaxRotationError = -1.0f;
 
 	// Whether to limit the offset's translation interpolation speed to the velocity on the incoming motion
 	// Enabling this prevents the offset sliding when there's little to no translation speed
-	UPROPERTY(EditAnywhere, BlueprintReadWrite, Category = Settings)
+	UPROPERTY(EditAnywhere, Category = Settings, meta = (FoldProperty, PinHiddenByDefault))
 	bool bClampToTranslationVelocity = false;
 
 	// Whether to limit the offset's rotation interpolation speed to the velocity on the incoming motion
 	// Enabling this prevents the offset sliding when there's little to no rotation speed
-	UPROPERTY(EditAnywhere, BlueprintReadWrite, Category = Settings)
+	UPROPERTY(EditAnywhere, Category = Settings, meta = (FoldProperty, PinHiddenByDefault))
 	bool bClampToRotationVelocity = false;
 
 	// How much the offset can blend out, relative to the incoming translation speed
 	// i.e. If root motion is moving at 400cm/s, at 0.5, the offset can blend out at 200cm/s
-	UPROPERTY(EditAnywhere, BlueprintReadWrite, Category = Settings, meta = (EditCondition = bClampToTranslationVelocity))
+	UPROPERTY(EditAnywhere, Category = Settings, meta = (EditCondition = bClampToTranslationVelocity, FoldProperty, PinHiddenByDefault))
 	float TranslationSpeedRatio = 0.5f;
 
 	// How much the offset can blend out, relative to the incoming rotation speed
 	// i.e. If root motion is rotating at 90 degrees/s, at 0.5, the offset can blend out at 45 degree/s
-	UPROPERTY(EditAnywhere, BlueprintReadWrite, Category = Settings, meta = (EditCondition = bClampToRotationVelocity))
+	UPROPERTY(EditAnywhere, Category = Settings, meta = (EditCondition = bClampToRotationVelocity, FoldProperty, PinHiddenByDefault))
 	float RotationSpeedRatio = 0.5f;
-};
-
-USTRUCT(BlueprintInternalUseOnly, Experimental)
-struct ANIMATIONWARPINGRUNTIME_API FAnimNode_OffsetRootBone : public FAnimNode_SkeletalControlBase
-{
-	GENERATED_BODY()
-
-#if WITH_EDITORONLY_DATA
-	UPROPERTY(EditAnywhere, Category = Evaluation, meta=(FoldProperty))
-	EWarpingEvaluationMode EvaluationMode = EWarpingEvaluationMode::Graph;
-
-	UPROPERTY(EditAnywhere, Category = Settings, meta = (FoldProperty, PinHiddenByDefault))
-	FOffsetRootBoneSettings OffsetSettings;
 
 	// Delta applied to the translation offset this frame. 
 	// For procedural values, consider adjusting the input by delta time.
@@ -116,13 +107,25 @@ public:
 	virtual bool IsValidToEvaluate(const USkeleton* Skeleton, const FBoneContainer& RequiredBones) override;
 	// End of FAnimNode_SkeletalControlBase interface
 
+	// Folded property accesors
 	EWarpingEvaluationMode GetEvaluationMode() const;
-	const FOffsetRootBoneSettings& GetOffsetSettings() const;
 	const FVector& GetTranslationDelta() const;
 	const FRotator& GetRotationDelta() const;
+	EOffsetRootBoneMode GetTranslationMode() const;
+	EOffsetRootBoneMode GetRotationMode() const;
+	float GetTranslationHalflife() const;
+	float GetRotationHalfLife() const;
+	float GetMaxTranslationError() const;
+	float GetMaxRotationError() const;
+	bool GetClampToTranslationVelocity() const;
+	bool GetClampToRotationVelocity() const;
+	float GetTranslationSpeedRatio() const;
+	float GetRotationSpeedRatio() const;
 
 private:
 
+	void Reset(const FAnimationBaseContext& Context);
+
 	// Internal cached anim instance proxy
 	FAnimInstanceProxy* AnimInstanceProxy = nullptr;
 
@@ -137,4 +140,6 @@ private:
 	// Offset = ComponentTransform - SimulatedTransform
 	FVector SimulatedTranslation = FVector::ZeroVector;
 	FQuat SimulatedRotation = FQuat::Identity;
+
+	FGraphTraversalCounter UpdateCounter;
 };
