Foot Placement: Fix for pelvis over-compensation when going up/down stairs

diff --git a/Engine/Plugins/Animation/AnimationWarping/Source/Runtime/Private/BoneControllers/AnimNode_FootPlacement.cpp b/Engine/Plugins/Animation/AnimationWarping/Source/Runtime/Private/BoneControllers/AnimNode_FootPlacement.cpp
--- a/Engine/Plugins/Animation/AnimationWarping/Source/Runtime/Private/BoneControllers/AnimNode_FootPlacement.cpp
+++ b/Engine/Plugins/Animation/AnimationWarping/Source/Runtime/Private/BoneControllers/AnimNode_FootPlacement.cpp
@@ -1406,7 +1406,7 @@ void FAnimNode_FootPlacement::ProcessCharacterState(const UE::Anim::FootPlacemen
 
 		// Only compensate vertical motion
 		const FVector CapsuleMoveOffsetWS =
-			(ComponentLocationWS - LastComponentLocationWS) * -Context.ApproachDirWS;
+			(ComponentLocationWS - OwningComponentAdjustedLastLocationWS) * -Context.ApproachDirWS;
 		if (!CapsuleMoveOffsetWS.IsNearlyZero(KINDA_SMALL_NUMBER))
 		{
 			const FVector CapsuleMoveOffsetCS =
