Foot Placement: Fix for toe popping when leg is near overextension and hypercompression

diff --git a/Engine/Plugins/Animation/AnimationWarping/Source/Runtime/Private/BoneControllers/AnimNode_FootPlacement.cpp b/Engine/Plugins/Animation/AnimationWarping/Source/Runtime/Private/BoneControllers/AnimNode_FootPlacement.cpp
--- a/Engine/Plugins/Animation/AnimationWarping/Source/Runtime/Private/BoneControllers/AnimNode_FootPlacement.cpp
+++ b/Engine/Plugins/Animation/AnimationWarping/Source/Runtime/Private/BoneControllers/AnimNode_FootPlacement.cpp
@@ -772,7 +772,7 @@ UE::Anim::FootPlacement::FPlantResult FAnimNode_FootPlacement::FinalizeFootAlign
 				// Try to keep the tip on spot if we're unplanting
 				// Don't do this until we've reached the plant target once
 				// TODO: Make this configurable? 
-				if ((bRecentlyUnplanted || bIsPlanted) && LegData.Plant.bCanReachTarget)
+				if ((bRecentlyUnplanted || (bIsPlanted && LegData.Plant.bCanReachTarget))
 				{
 					const FTransform FKHipToLeg =
 						LegData.InputPose.FootTransformCS.GetRelativeTransform(PelvisData.InputPose.FKTransformCS);
@@ -1664,7 +1664,7 @@ FTransform FAnimNode_FootPlacement::SolvePelvis(const UE::Anim::FootPlacement::F
 		DesiredOffsetMin + ((MinToAvg * MinToMax) / Divisor);
 
 	// Adjust the hips to prevent over-compression
-	//PelvisOffsetZ = FMath::Clamp(PelvisOffsetZ, MinOffsetMax, MaxOffsetMin);
+	PelvisOffsetZ = FMath::Clamp(PelvisOffsetZ, MinOffsetMax, MaxOffsetMin);
 
 	FVector PelvisOffsetDelta = -PelvisOffsetZ * Context.ApproachDirCS;
 
