Added root bone offset to orientation warping, to follow root motion rotation more closely during forward-facing turns.

diff --git a/Engine/Plugins/Animation/AnimationWarping/Source/Runtime/Private/BoneControllers/AnimNode_OrientationWarping.cpp b/Engine/Plugins/Animation/AnimationWarping/Source/Runtime/Private/BoneControllers/AnimNode_OrientationWarping.cpp
--- a/Engine/Plugins/Animation/AnimationWarping/Source/Runtime/Private/BoneControllers/AnimNode_OrientationWarping.cpp
+++ b/Engine/Plugins/Animation/AnimationWarping/Source/Runtime/Private/BoneControllers/AnimNode_OrientationWarping.cpp
@@ -7,7 +7,7 @@
 DECLARE_CYCLE_STAT(TEXT("OrientationWarping Eval"), STAT_OrientationWarping_Eval, STATGROUP_Anim);
 
 #if ENABLE_ANIM_DEBUG
-static TAutoConsoleVariable<int32> CVarAnimNodeOrientationWarpingDebug(TEXT("a.AnimNode.OrientationWarping.Debug"), 0, TEXT("Turn on visualization debugging for Orientation Warping"));
+static TAutoConsoleVariable<int32> CVarAnimNodeOrientationWarpingDebug(TEXT("a.AnimNode.OrientationWarping.Debug"), 0, TEXT("Turn on visualization debugging for Orientation Warping. 0 == off, 1 == Warping, 2 == Offset, 3 == Both."));
 static TAutoConsoleVariable<int32> CVarAnimNodeOrientationWarpingVerbose(TEXT("a.AnimNode.OrientationWarping.Verbose"), 0, TEXT("Turn on verbose graph debugging for Orientation Warping"));
 static TAutoConsoleVariable<int32> CVarAnimNodeOrientationWarpingEnable(TEXT("a.AnimNode.OrientationWarping.Enable"), 1, TEXT("Toggle Orientation Warping"));
 #endif
@@ -55,7 +55,9 @@ void FAnimNode_OrientationWarping::GatherDebugData(FNodeDebugData& DebugData)
 			DebugLine += FString::Printf(TEXT("\n - Root Motion Delta Attribute Found: %s)"), (bFoundRootMotionAttribute) ? TEXT("true") : TEXT("false"));
 #endif
 		}
-		DebugLine += FString::Printf(TEXT("\n - Distributed Bone Orientation Alpha: (%.3fd)"), DistributedBoneOrientationAlpha);
+		DebugLine += FString::Printf(TEXT("\n - Warping Alpha: (%.2fd)"), WarpingAlpha);
+		DebugLine += FString::Printf(TEXT("\n - Heading Offset: (%.3fd)"), FMath::RadiansToDegrees(HeadingOffset));
+		DebugLine += FString::Printf(TEXT("\n - Offset Alpha: (%.2fd)"), OffsetAlpha);
 		if (const UEnum* TypeEnum = FindObject<UEnum>(nullptr, TEXT("/Script/CoreUObject.EAxis")))
 		{
 			DebugLine += FString::Printf(TEXT("\n - Rotation Axis: (%s)"), *(TypeEnum->GetNameStringByIndex(static_cast<int32>(RotationAxis))));
@@ -65,7 +67,9 @@ void FAnimNode_OrientationWarping::GatherDebugData(FNodeDebugData& DebugData)
 	else
 #endif
 	{
-	DebugLine += FString::Printf(TEXT("(Orientation Angle: %.3fd)"), FMath::RadiansToDegrees(ActualOrientationAngle));
+	const float ActualOrientationAngleDegrees = FMath::RadiansToDegrees(ActualOrientationAngle);
+	const float HeadingOffsetDegrees = FMath::RadiansToDegrees(HeadingOffset);
+	DebugLine += FString::Printf(TEXT("(Orientation Angle: %.3fd, Heading Offset : %.3fd, )"), ActualOrientationAngleDegrees, HeadingOffsetDegrees);
 	}
 	DebugData.AddDebugItem(DebugLine);
 	ComponentPose.GatherDebugData(DebugData);
@@ -78,11 +82,19 @@ void FAnimNode_OrientationWarping::Initialize_AnyThread(const FAnimationInitiali
 	PreviousRootMotionDeltaDirection = FVector::ZeroVector;
 	PreviousOrientationAngle = 0.f;
 	ActualOrientationAngle = 0.f;
+	Reset(Context);
 }
 
 void FAnimNode_OrientationWarping::UpdateInternal(const FAnimationUpdateContext& Context)
 {
 	FAnimNode_SkeletalControlBase::UpdateInternal(Context);
+
+	// If we just became relevant and haven't been initialized yet, then reset.
+	if (!bIsFirstUpdate && UpdateCounter.HasEverBeenUpdated() && !UpdateCounter.WasSynchronizedCounter(Context.AnimInstanceProxy->GetUpdateCounter()))
+	{
+		Reset(Context);
+	}
+	UpdateCounter.SynchronizeWith(Context.AnimInstanceProxy->GetUpdateCounter());
 }
 
 void FAnimNode_OrientationWarping::EvaluateSkeletalControl_AnyThread(FComponentSpacePoseContext& Output, TArray<FBoneTransform>& OutBoneTransforms)
@@ -153,8 +165,41 @@ void FAnimNode_OrientationWarping::EvaluateSkeletalControl_AnyThread(FComponentS
 			const FTransform SkeletalMeshRelativeTransform = Output.AnimInstanceProxy->GetComponentRelativeTransform();
 			const FQuat SkeletalMeshRelativeRotation = SkeletalMeshRelativeTransform.GetRotation();
 			LocomotionForward = SkeletalMeshRelativeRotation.UnrotateVector(LocomotionRotation.GetForwardVector()).GetSafeNormal();
+			
+			if (OffsetAlpha == 0.0)
+			{
+				// No alpha. Reset our heading offset.
+				HeadingOffset = 0.0f;
+			}
+			else
+			{
+				// Accumulate a percentage of our component's frame delta, based on offset alpha
+				const float LastComponentHeading = ComponentHeading;
+				ComponentHeading = Output.AnimInstanceProxy->GetComponentTransform().GetRotation().GetTwistAngle(RotationAxisVector);
+				HeadingOffset += FMath::UnwindRadians(LastComponentHeading - ComponentHeading) * OffsetAlpha;
+
+				// Accumulate a percentage of our root motion's heading, based on offset alpha
+				float RootMotionDeltaHeading = RootMotionTransformDelta.GetRotation().GetTwistAngle(RotationAxisVector);
+				HeadingOffset += RootMotionDeltaHeading * OffsetAlpha;
+
+				// If our alpha is decreasing, we're blending out the offset.
+				// We don't need to do this when blending in because we're just accumulating partial component offset and root motion.
+				const float OffsetAlphaStep = LastOffsetAlpha - OffsetAlpha;
+				if (OffsetAlphaStep > 0.0f)
+				{
+					HeadingOffset -= OffsetAlphaStep * HeadingOffset / LastOffsetAlpha;
+				}
+
+				const float MaxOffsetRadians = FMath::DegreesToRadians(MaxOffsetAngle);
+				HeadingOffset = FMath::Clamp(HeadingOffset, -MaxOffsetAngle, MaxOffsetAngle);
+			}
+			LastOffsetAlpha = OffsetAlpha;
+
+			// Rotate the root motion direction by our effective offset.
+			// This means we will warp our pose based on the adjusted offset, and correct accordingly.
+			FQuat RootOffsetRotation = FQuat(RotationAxisVector, HeadingOffset);
+			const FVector RootMotionDeltaTranslation = RootOffsetRotation.RotateVector(RootMotionTransformDelta.GetTranslation());
 
-			const FVector RootMotionDeltaTranslation = RootMotionTransformDelta.GetTranslation();
 			// Hold previous direction if we can't calculate it from current move delta, because the root is no longer moving
 			RootMotionDeltaDirection = RootMotionDeltaTranslation.GetSafeNormal(UE_SMALL_NUMBER, PreviousRootMotionDeltaDirection);
 
@@ -229,7 +274,7 @@ void FAnimNode_OrientationWarping::EvaluateSkeletalControl_AnyThread(FComponentS
 	}
 
 	// Allow the alpha value of the node to affect the final rotation
-	ActualOrientationAngle *= ActualAlpha;
+	ActualOrientationAngle *= ActualAlpha * WarpingAlpha;
 
 #if ENABLE_ANIM_DEBUG
 	bool bDebugging = false;
@@ -238,50 +283,78 @@ void FAnimNode_OrientationWarping::EvaluateSkeletalControl_AnyThread(FComponentS
 #else
 	constexpr float DebugDrawScale = 1.f;
 #endif
-	bDebugging = bDebugging || CVarAnimNodeOrientationWarpingDebug.GetValueOnAnyThread() == 1;
+	const int32 DebugIndex = CVarAnimNodeOrientationWarpingDebug.GetValueOnAnyThread();
+	bDebugging = bDebugging || (DebugIndex > 0);
 
 	if (bDebugging)
 	{
 		const FTransform ComponentTransform = Output.AnimInstanceProxy->GetComponentTransform();
 		const FVector ActorForwardDirection = Output.AnimInstanceProxy->GetActorTransform().GetRotation().GetForwardVector();
-
-		const FVector ForwardDirection = bGraphDrivenWarping
-			? ComponentTransform.GetRotation().RotateVector(LocomotionForward)
-			: ActorForwardDirection;
-
 		FVector DebugArrowOffset = FVector::ZAxisVector * DebugDrawScale;
-		Output.AnimInstanceProxy->AnimDrawDebugDirectionalArrow(
-			ComponentTransform.GetLocation() + DebugArrowOffset,
-			ComponentTransform.GetLocation() + DebugArrowOffset + ForwardDirection * 100.f * DebugDrawScale,
-			40.f * DebugDrawScale, FColor::Red, false, 0.f, 2.f * DebugDrawScale);
-
-		const FVector RotationDirection = bGraphDrivenWarping
-			? ComponentTransform.GetRotation().RotateVector(RootMotionDeltaDirection)
-			: ActorForwardDirection.RotateAngleAxis(OrientationAngle, RotationAxisVector);
-
-		DebugArrowOffset += DebugArrowOffset;
-		Output.AnimInstanceProxy->AnimDrawDebugDirectionalArrow(
-			ComponentTransform.GetLocation() + DebugArrowOffset,
-			ComponentTransform.GetLocation() + DebugArrowOffset + RotationDirection * 100.f * DebugDrawScale,
-			40.f * DebugDrawScale, FColor::Blue, false, 0.f, 2.f * DebugDrawScale);
-
-		const float ActualOrientationAngleDegrees = FMath::RadiansToDegrees(ActualOrientationAngle);
-		const FVector WarpedRotationDirection = bGraphDrivenWarping 
-			? RotationDirection.RotateAngleAxis(ActualOrientationAngleDegrees, RotationAxisVector)
-			: ActorForwardDirection.RotateAngleAxis(ActualOrientationAngleDegrees, RotationAxisVector);
-
-		DebugArrowOffset += DebugArrowOffset;
-		Output.AnimInstanceProxy->AnimDrawDebugDirectionalArrow(
-			ComponentTransform.GetLocation() + DebugArrowOffset,
-			ComponentTransform.GetLocation() + DebugArrowOffset + WarpedRotationDirection * 100.f * DebugDrawScale,
-			40.f * DebugDrawScale, FColor::Green, false, 0.f, 2.f * DebugDrawScale);
+
+		const bool bDrawAll = DebugIndex == 3;
+		const bool bDrawOffset = (DebugIndex == 2) || bDrawAll;
+		const bool bDrawWarping = (DebugIndex == 1) || bDrawAll;
+		if (bDrawWarping)
+		{
+			const FVector ForwardDirection = bGraphDrivenWarping
+				? ComponentTransform.GetRotation().RotateVector(LocomotionForward)
+				: ActorForwardDirection;
+
+			Output.AnimInstanceProxy->AnimDrawDebugDirectionalArrow(
+				ComponentTransform.GetLocation() + DebugArrowOffset,
+				ComponentTransform.GetLocation() + DebugArrowOffset + ForwardDirection * 100.f * DebugDrawScale,
+				40.f * DebugDrawScale, FColor::Red, false, 0.f, 2.f * DebugDrawScale);
+
+			const FVector RotationDirection = bGraphDrivenWarping
+				? ComponentTransform.GetRotation().RotateVector(RootMotionDeltaDirection)
+				: ActorForwardDirection.RotateAngleAxis(OrientationAngle, RotationAxisVector);
+
+			DebugArrowOffset += FVector::ZAxisVector * DebugDrawScale;
+			Output.AnimInstanceProxy->AnimDrawDebugDirectionalArrow(
+				ComponentTransform.GetLocation() + DebugArrowOffset,
+				ComponentTransform.GetLocation() + DebugArrowOffset + RotationDirection * 100.f * DebugDrawScale,
+				40.f * DebugDrawScale, FColor::Blue, false, 0.f, 2.f * DebugDrawScale);
+
+			const float ActualOrientationAngleDegrees = FMath::RadiansToDegrees(ActualOrientationAngle);
+			const FVector WarpedRotationDirection = bGraphDrivenWarping
+				? RotationDirection.RotateAngleAxis(ActualOrientationAngleDegrees, RotationAxisVector)
+				: ActorForwardDirection.RotateAngleAxis(ActualOrientationAngleDegrees, RotationAxisVector);
+
+			DebugArrowOffset += FVector::ZAxisVector * DebugDrawScale;
+			Output.AnimInstanceProxy->AnimDrawDebugDirectionalArrow(
+				ComponentTransform.GetLocation() + DebugArrowOffset,
+				ComponentTransform.GetLocation() + DebugArrowOffset + WarpedRotationDirection * 100.f * DebugDrawScale,
+				40.f * DebugDrawScale, FColor::Green, false, 0.f, 2.f * DebugDrawScale);
+			DebugArrowOffset += FVector::ZAxisVector * DebugDrawScale;
+		}
+		
+		if (bDrawOffset)
+		{
+			const float HeadingOffetDegrees = FMath::RadiansToDegrees(HeadingOffset);
+			const FVector OffsetDirection = ActorForwardDirection.RotateAngleAxis(HeadingOffetDegrees, RotationAxisVector);
+
+			Output.AnimInstanceProxy->AnimDrawDebugDirectionalArrow(
+				ComponentTransform.GetLocation() + DebugArrowOffset,
+				ComponentTransform.GetLocation() + DebugArrowOffset + OffsetDirection * 100.f * DebugDrawScale,
+				40.f * DebugDrawScale, FColor::Purple, false, 0.f, 2.f * DebugDrawScale);
+
+			DebugArrowOffset += FVector::ZAxisVector * DebugDrawScale;
+			Output.AnimInstanceProxy->AnimDrawDebugDirectionalArrow(
+				ComponentTransform.GetLocation() + DebugArrowOffset,
+				ComponentTransform.GetLocation() + DebugArrowOffset + ActorForwardDirection * 100.f * DebugDrawScale,
+				40.f * DebugDrawScale, FColor::Black, false, 0.f, 2.f * DebugDrawScale);
+		}
 	}
 #endif
-		
+
+	// Combine our warping and heading offsets to rotate our root bone
+	const float CombinedRootOffset = FMath::UnwindRadians(ActualOrientationAngle * DistributedBoneOrientationAlpha + HeadingOffset);
+
 	// Rotate Root Bone first, as that cheaply rotates the whole pose with one transformation.
-	if (!FMath::IsNearlyZero(DistributedBoneOrientationAlpha, KINDA_SMALL_NUMBER))
+	if (!FMath::IsNearlyZero(CombinedRootOffset, KINDA_SMALL_NUMBER))
 	{
-		const FQuat RootRotation = FQuat(RotationAxisVector, ActualOrientationAngle * DistributedBoneOrientationAlpha);
+		const FQuat RootRotation = FQuat(RotationAxisVector, CombinedRootOffset);
 		const FCompactPoseBoneIndex RootBoneIndex(0);
 
 		FTransform RootBoneTransform(Output.Pose.GetComponentSpaceTransform(RootBoneIndex));
@@ -468,4 +541,11 @@ void FAnimNode_OrientationWarping::InitializeBoneReferences(const FBoneContainer
 			}
 		}
 	}
-}
\ No newline at end of file
+}
+
+void FAnimNode_OrientationWarping::Reset(const FAnimationBaseContext& Context)
+{
+	bIsFirstUpdate = true;
+	HeadingOffset = 0.0f;
+	ComponentHeading = Context.AnimInstanceProxy->GetComponentTransform().GetRotation().GetTwistAngle(UE::Anim::GetAxisVector(RotationAxis));
+}
diff --git a/Engine/Plugins/Animation/AnimationWarping/Source/Runtime/Public/BoneControllers/AnimNode_OrientationWarping.h b/Engine/Plugins/Animation/AnimationWarping/Source/Runtime/Public/BoneControllers/AnimNode_OrientationWarping.h
--- a/Engine/Plugins/Animation/AnimationWarping/Source/Runtime/Public/BoneControllers/AnimNode_OrientationWarping.h
+++ b/Engine/Plugins/Animation/AnimationWarping/Source/Runtime/Public/BoneControllers/AnimNode_OrientationWarping.h
@@ -67,6 +67,15 @@ struct ANIMATIONWARPINGRUNTIME_API FAnimNode_OrientationWarping : public FAnimNo
 	UPROPERTY(EditAnywhere, Category=Settings, meta=(ClampMin="0.0"))
 	float RotationInterpSpeed = 10.f;
 
+	UPROPERTY(EditAnywhere, Category = Experimental, meta=(PinHiddenByDefault, ClampMin="0.0", ClampMax="1.0"))
+	float WarpingAlpha = 1.0f;
+
+	UPROPERTY(EditAnywhere, Category = Experimental, meta=(PinHiddenByDefault, ClampMin="0.0", ClampMax="1.0"))
+	float OffsetAlpha = 0.0f;
+
+	UPROPERTY(EditAnywhere, Category = Experimental, meta=(PinHiddenByDefault, ClampMin="0.0", ClampMax="180.0"))
+	float MaxOffsetAngle = 70.0f;
+
 #if WITH_EDITORONLY_DATA
 	// Scale all debug drawing visualization by a factor
 	UPROPERTY(EditAnywhere, Category=Debug, meta=(ClampMin="0.0"))
@@ -148,6 +157,17 @@ private:
 	// Internal orientation warping angle
 	float ActualOrientationAngle = 0.f;
 
+	// Our component's heading
+	float ComponentHeading = 0.0;
+
+	// Our accumulated heading offset based on component frame deltas
+	float HeadingOffset = 0.0;
+	float LastOffsetAlpha = 0.0f;
+
+	FGraphTraversalCounter UpdateCounter;
+	bool bIsFirstUpdate = false;
+	void Reset(const FAnimationBaseContext& Context);
+
 #if WITH_EDITORONLY_DATA
 	// Whether we found a root motion delta attribute in the attribute stream on graph driven mode
 	bool bFoundRootMotionAttribute = false;
