Foot Placement:
-Improved pelvis displacement during foot locking
-Cleaned up the foot rolling and pelvis adjustments

diff --git a/Engine/Plugins/Animation/AnimationWarping/Source/Runtime/Private/BoneControllers/AnimNode_FootPlacement.cpp b/Engine/Plugins/Animation/AnimationWarping/Source/Runtime/Private/BoneControllers/AnimNode_FootPlacement.cpp
--- a/Engine/Plugins/Animation/AnimationWarping/Source/Runtime/Private/BoneControllers/AnimNode_FootPlacement.cpp
+++ b/Engine/Plugins/Animation/AnimationWarping/Source/Runtime/Private/BoneControllers/AnimNode_FootPlacement.cpp
@@ -240,7 +240,6 @@ namespace UE::Anim::FootPlacement
 		{
 			FCompactPoseBoneIndex ParentBoneIndex = RequiredBones.GetParentBoneIndex(BoneIndex);
 
-
 			int32 NumIterations = NumBonesInLimb;
 			while ((NumIterations-- > 0) && (ParentBoneIndex != INDEX_NONE))
 			{
@@ -259,12 +258,16 @@ namespace UE::Anim::FootPlacement
 
 void FAnimNode_FootPlacement::FindPelvisOffsetRangeForLimb(
 	const UE::Anim::FootPlacement::FEvaluationContext& Context,
-	const UE::Anim::FootPlacement::FLegRuntimeData::FInputPoseData& LegInputPose,
-	const FVector& PlantTargetLocationCS,
+	const UE::Anim::FootPlacement::FLegRuntimeData& LegData,
+	const FVector& InPlantTargetLocationCS,
 	const FTransform& PelvisTransformCS,
-	const float LimbLength,
 	FPelvisOffsetRangeForLimb& OutPelvisOffsetRangeCS) const
 {
+	const UE::Anim::FootPlacement::FLegRuntimeData::FInputPoseData& LegInputPose = LegData.InputPose;
+	const UE::Anim::FootPlacement::FLegRuntimeData::FBoneData& Bones = LegData.Bones;
+	const float LimbLength = Bones.LimbLength;
+	FVector PlantTargetLocationCS = InPlantTargetLocationCS;
+
 	// TODO: Cache this.
 	const FTransform HipToPelvis =
 		LegInputPose.HipTransformCS.GetRelativeTransform(PelvisData.InputPose.FKTransformCS);
@@ -278,56 +281,65 @@ void FAnimNode_FootPlacement::FindPelvisOffsetRangeForLimb(
 	const float DesiredExtension = FMath::Sqrt(DesiredExtensionSqrd);
 	const float MaxExtension = GetMaxLimbExtension(DesiredExtension, LimbLength);
 
+	const FPlane FootPlane = FPlane(PlantTargetLocationCS, -Context.ApproachDirCS);
 	const FVector HipToPlantCS = PlantTargetLocationCS - HipLocationCS;
-	const float HipToPlantDotApproachDir = HipToPlantCS | Context.ApproachDirCS;
+	const float HipHeight = FootPlane.PlaneDot(HipLocationCS);
 
 	FVector DesiredPlantTargetLocationCS = PlantTargetLocationCS;
 	FVector MaxPlantTargetLocationCS = PlantTargetLocationCS;
-	// If the foot wants to be place so high up relative to the FK hip, this is unlikely to matter.
-	if (HipToPlantDotApproachDir > 0.0f)
+	// Don't do horizontal adjustments if the foot is above the hip
+	if (HipHeight > 0.0f)
 	{
-		const float OpposingSideSqrd = HipToPlantCS.SizeSquared() - HipToPlantDotApproachDir * HipToPlantDotApproachDir;
-		const float OpposingSide = FMath::Sqrt(OpposingSideSqrd);
-
-		const float MaxExtensionSqrd = MaxExtension * MaxExtension;
-
-		const FPlane FootPlane = FPlane(PlantTargetLocationCS, Context.ApproachDirCS);
+		// Project the input pose foot and hip to the ground-aligned foot plane
 		const FVector FKFootProjected = FVector::PointPlaneProject(LegInputPose.FootTransformCS.GetLocation(), FootPlane);
 		const FVector HipProjected = FVector::PointPlaneProject(HipLocationCS, FootPlane);
+		
+		// Project both our FK and IK feet to the foot plane, and calculate distances to the projected hip.
+		FVector TargetFootToHip = HipProjected - PlantTargetLocationCS;
+		float TargetFootOffset = FVector::Dist(HipProjected, PlantTargetLocationCS);
+		float InitialFootOffset = FVector::Dist(HipProjected, FKFootProjected);
+
+		// Move our IK foot to our FK foot by about the foot's length horizontally.
+		// This will make our heel lift before it drops the hips when locking.
+		//@TODO: This is not a very accurate estimate, but neither is our foot roll yet. Fix this when we improve the foot roll.
+		const FVector ToFKFoot = (FKFootProjected - PlantTargetLocationCS);
+		const FVector ToFKFootDir = ToFKFoot.GetSafeNormal();
+		const float ToFKDistance = ToFKFoot.Size();
+		PlantTargetLocationCS = PlantTargetLocationCS + ToFKFootDir * FMath::Min(ToFKDistance, Bones.FootLength);
+
+		// Calculate new correction offset and direction
+		TargetFootToHip = HipProjected - PlantTargetLocationCS;
+		TargetFootOffset = FVector::Dist(HipProjected, PlantTargetLocationCS);
 
-		const float MaxOffset = PelvisSettings.MaxOffsetHorizontal;
-		const FVector IKFootToHip = HipProjected - PlantTargetLocationCS;
-		const float IKFootToHipDist = IKFootToHip.Size();
-		const float FKFootToHipDist = FVector::Dist(HipProjected, FKFootProjected);
 		auto FindPlantLocationAdjustedByOrtogonalLimit = 
-			[	PlantTargetLocationCS, MaxOffset, IKFootToHip, IKFootToHipDist,
-				HipToPlantDotApproachDir, OpposingSide, FKFootToHipDist	]
-			(const float RadiusSqrd)
+			[	TargetFootOffset, TargetFootToHip,
+				HipHeight, InitialFootOffset	]
+			(const float LegLength, const float MaxHipOffset, const FVector& FootLocation)
 		{
-			FVector AdjustedPlantTargetLocationCS = PlantTargetLocationCS;
+			FVector AdjustedPlantTargetLocationCS = FootLocation;
 
-			// The desired height at the limit of  max horizontal extension
-			const float DesiredHeight = HipToPlantDotApproachDir - MaxOffset;
-			const float DesiredHeightSqrd = DesiredHeight * DesiredHeight;
+			// The minimum height our hip can be at after horizontal adjustments
+			const float MinHeight = HipHeight - MaxHipOffset;
+			const float MinHeightSqrd = FMath::Square(MinHeight);
 
-			// Find the max horizontal offset. 
-			// We don't care about the circle intersection in the opposite direction
-			const float MaxOpposingSide = FMath::Sqrt(FMath::Abs(-DesiredHeightSqrd + RadiusSqrd));
+			// Find how far our foot would be from the projected hip, if the leg was at max extension. 
+			const float LegLengthSqrd = FMath::Square(LegLength);
+			const float MaxFootOffset = FMath::Sqrt(FMath::Max(0.0f, LegLengthSqrd - MinHeightSqrd));
 
-			// Respected the input pose if it exceeds it
-			const float MaxIKOrthogonalDist = FMath::Max(FKFootToHipDist, MaxOpposingSide);
+			// If the input pose is already further than this, respect the input pose
+			const float MaxFootOffsetClamped = FMath::Max(InitialFootOffset, MaxFootOffset);
 
-			if (IKFootToHipDist > MaxIKOrthogonalDist)
+			if (TargetFootOffset > MaxFootOffsetClamped)
 			{
 				// Move the foot towards the projected hip
-				AdjustedPlantTargetLocationCS += (IKFootToHipDist - MaxIKOrthogonalDist) * IKFootToHip.GetSafeNormal();
+				AdjustedPlantTargetLocationCS += (TargetFootOffset - MaxFootOffsetClamped) * TargetFootToHip.GetSafeNormal();
 			}
 
 			return AdjustedPlantTargetLocationCS;
 		};
 
-		MaxPlantTargetLocationCS = FindPlantLocationAdjustedByOrtogonalLimit(MaxExtensionSqrd);
-		DesiredPlantTargetLocationCS = FindPlantLocationAdjustedByOrtogonalLimit(DesiredExtensionSqrd);
+		MaxPlantTargetLocationCS = FindPlantLocationAdjustedByOrtogonalLimit(MaxExtension, PelvisSettings.MaxOffsetHorizontal, PlantTargetLocationCS);
+		DesiredPlantTargetLocationCS = FindPlantLocationAdjustedByOrtogonalLimit(DesiredExtension,  PelvisSettings.MaxOffsetHorizontal, PlantTargetLocationCS);
 	}
 
 	// Taken from http://runevision.com/thesis/rune_skovbo_johansen_thesis.pdf
@@ -732,30 +744,27 @@ UE::Anim::FootPlacement::FPlantResult FAnimNode_FootPlacement::FinalizeFootAlign
 		LegData.InputPose.HipTransformCS.GetRelativeTransform(PelvisData.InputPose.FKTransformCS);
 	const FTransform FinalHipTransformCS = FKPelvisToHipCS * PelvisTransformCS;
 	FTransform CorrectedFootTransformCS = LegData.AlignedFootTransformCS;
+	const FVector CorrectedBallLocationCS = (LegData.InputPose.FootToBall * CorrectedFootTransformCS).GetLocation();
 
 	// avoid hyper extension - start
-	const FVector InitialHipToFootDir =
-		(LegData.InputPose.FootTransformCS.GetLocation() - LegData.InputPose.HipTransformCS.GetLocation()).GetSafeNormal();
-	const FVector TargetHipToFootDir =
-		(CorrectedFootTransformCS.GetLocation() - FinalHipTransformCS.GetLocation()).GetSafeNormal();
-	
-
-	const FVector FKHipToFoot =
+	const FVector InputPoseHipToFoot =
 		LegData.InputPose.FootTransformCS.GetLocation() - LegData.InputPose.HipTransformCS.GetLocation();
+	const FVector InputPoseHipToFootDir = InputPoseHipToFoot.GetSafeNormal();
+	const FVector CorrectedHipToFootDir =
+		(CorrectedFootTransformCS.GetLocation() - FinalHipTransformCS.GetLocation()).GetSafeNormal();
 
-	if (!InitialHipToFootDir.IsNearlyZero() && !TargetHipToFootDir.IsNearlyZero())
+	if (!InputPoseHipToFoot.IsNearlyZero() && !CorrectedHipToFootDir.IsNearlyZero())
 	{
 		{
-
-			const float FKExtension = FKHipToFoot.Size();
-			const float MaxExtension = GetMaxLimbExtension(FKExtension, LegData.Bones.LimbLength);
-			const float IKExtension =
+			const float DesiredExtension = InputPoseHipToFoot.Size();
+			const float MaxExtension = GetMaxLimbExtension(DesiredExtension, LegData.Bones.LimbLength);
+			const float CurrentExtension =
 				FVector::Dist(CorrectedFootTransformCS.GetLocation(), FinalHipTransformCS.GetLocation());
 
-			const float HyperExtensionAmount = IKExtension - MaxExtension;
+			const float HyperExtensionAmount = CurrentExtension - MaxExtension;
 			float HyperExtensionRemaining = HyperExtensionAmount;
 
-			if (IKExtension > MaxExtension)
+			if (CurrentExtension > MaxExtension)
 			{
 				const bool bIsPlanted = LegData.Plant.PlantType != UE::Anim::FootPlacement::EPlantType::Unplanted;
 				const bool bWasPlanted = LegData.Plant.LastPlantType != UE::Anim::FootPlacement::EPlantType::Unplanted;
@@ -774,36 +783,28 @@ UE::Anim::FootPlacement::FPlantResult FAnimNode_FootPlacement::FinalizeFootAlign
 				// TODO: Make this configurable? 
 				if (bRecentlyUnplanted || (bIsPlanted && LegData.Plant.bCanReachTarget))
 				{
-					const FTransform FKHipToLeg =
-						LegData.InputPose.FootTransformCS.GetRelativeTransform(PelvisData.InputPose.FKTransformCS);
-					const FTransform FKLegAtCurrentHipCS =
-						FKHipToLeg * PelvisTransformCS;
-					FVector IKToFK = (CorrectedFootTransformCS.GetLocation() - FKLegAtCurrentHipCS.GetLocation()).GetSafeNormal();
-
 					// Scale this value by our FK transition alpha to not pop
 					const float MaxPullTowardsHip =  FMath::Min(LegData.Bones.FootLength, HyperExtensionRemaining) * LegData.InputPose.AlignmentAlpha;
 					HyperExtensionRemaining -= MaxPullTowardsHip;
 
-					const FVector NotHyperextendedPlantLocation = CorrectedFootTransformCS.GetLocation() - TargetHipToFootDir * MaxPullTowardsHip;
-					// Grab the Tip location before adjustments.
-					const FVector HyperextendedIkTipLocationCS = (LegData.InputPose.FootToBall * CorrectedFootTransformCS).GetLocation();
+					//@TODO: This isn't accurate. Roll the foot around around the toes, and counter-adjust toe rotation instead
+					const FVector CorrectedFootLocationCS = CorrectedFootTransformCS.GetLocation() - MaxPullTowardsHip * CorrectedHipToFootDir;
 
-					const FTransform IKBallTransformCorrectedCS = LegData.InputPose.FootToBall * CorrectedFootTransformCS;
-
-					// Try to keep the tip at the same spot
-					const FVector ToTipInitial = (IKBallTransformCorrectedCS.GetLocation() - CorrectedFootTransformCS.GetLocation()).GetSafeNormal();
-					const FVector ToTipDesired = (IKBallTransformCorrectedCS.GetLocation() - NotHyperextendedPlantLocation).GetSafeNormal();
-					FQuat DeltaSlopeRotation = FQuat::FindBetweenNormals(ToTipInitial, ToTipDesired);
+					// Rotate the foot to keep the toe at the same spot
+					const FVector InitialFootToToe = CorrectedBallLocationCS - CorrectedFootTransformCS.GetLocation();
+					const FVector CorrectedFootToToe = CorrectedBallLocationCS - CorrectedFootLocationCS;
+					FQuat DeltaSlopeRotation = FQuat::FindBetweenVectors(InitialFootToToe, CorrectedFootToToe);
 
+#if ENABLE_ANIM_DEBUG
 					FRotator DeltaSlopeRotator = DeltaSlopeRotation.Rotator();
+#endif
 
+					// Rotate the foot to preserve the ball's location
 					CorrectedFootTransformCS.SetRotation(DeltaSlopeRotation * CorrectedFootTransformCS.GetRotation());
 					CorrectedFootTransformCS.NormalizeRotation();
 
-					// Move the IK bone closer to prevent overextension
-					CorrectedFootTransformCS.SetLocation(NotHyperextendedPlantLocation);
-					const FVector NotHyperextendedIkTipLocationCS = (LegData.InputPose.FootToBall * CorrectedFootTransformCS).GetLocation();
-					const float BallDelta = FVector::Dist(HyperextendedIkTipLocationCS, NotHyperextendedIkTipLocationCS);
+					// Move the foot bone closer to the hip prevent overextension
+					CorrectedFootTransformCS.SetLocation(CorrectedFootLocationCS);
 				}
 
 				// Fix any remaining hyper-extension
@@ -813,7 +814,7 @@ UE::Anim::FootPlacement::FPlantResult FAnimNode_FootPlacement::FinalizeFootAlign
 					// TODO: Pull towards the FK bone? This pull lifts the foot from the ground and it might be 
 					// preferable to slide. This causes discontinuities when the foot is no longer hyper-extended
 					FVector NotHyperextendedPlantLocation;
-					FMath::SphereDistToLine(FinalHipTransformCS.GetLocation(), MaxExtension, CorrectedFootTransformCS.GetLocation(), TargetHipToFootDir, NotHyperextendedPlantLocation);
+					FMath::SphereDistToLine(FinalHipTransformCS.GetLocation(), MaxExtension, CorrectedFootTransformCS.GetLocation(), CorrectedHipToFootDir, NotHyperextendedPlantLocation);
 					CorrectedFootTransformCS.SetLocation(NotHyperextendedPlantLocation);
 				}
 			}
@@ -834,11 +835,9 @@ UE::Anim::FootPlacement::FPlantResult FAnimNode_FootPlacement::FinalizeFootAlign
 	// Next the plant is ajusted to prevent penetration with the planting plane. To do that, first the base of the plant
 	// and the tip must be calculated (note that because the ground plane interpolates, this does not prevent physical penetration 
 	// with the geometry).
-	const FTransform CorrectedBallTransformCS = LegData.InputPose.FootToBall * CorrectedFootTransformCS;
-
 	// TODO: Consolidate with CalcTargetPlantPlaneDistance
 	const float FootDistance = UE::Anim::FootPlacement::GetDistanceToPlaneAlongDirection(
-		CorrectedBallTransformCS.GetLocation(),
+		CorrectedBallLocationCS,
 		LegData.Plant.PlantPlaneCS,
 		Context.ApproachDirCS);
 	const float BallDistance = UE::Anim::FootPlacement::GetDistanceToPlaneAlongDirection(
@@ -857,7 +856,7 @@ UE::Anim::FootPlacement::FPlantResult FAnimNode_FootPlacement::FinalizeFootAlign
 	// Doing this after pushing the feet out of the ground plane ensures we won't end up in awkward poses. 
 	{
 		FVector NotHyperextendedPlantLocation;
-		const float MinExtension = GetMinLimbExtension(FMath::Abs(FKHipToFoot | Context.ApproachDirCS), LegData.Bones.LimbLength);
+		const float MinExtension = GetMinLimbExtension(FMath::Abs(InputPoseHipToFoot | Context.ApproachDirCS), LegData.Bones.LimbLength);
 
 		// Offset our hip plane by min extension
 		FPlane HipPlane = FPlane(FinalHipTransformCS.GetLocation() + Context.ApproachDirCS * MinExtension, Context.ApproachDirCS);
@@ -873,15 +872,12 @@ UE::Anim::FootPlacement::FPlantResult FAnimNode_FootPlacement::FinalizeFootAlign
 
 	check(!CorrectedFootTransformCS.ContainsNaN());
 
-	// TODO: Do adjustments to FK tip and FK Chain root
-	FTransform FinalFkTipTransformCS = FTransform::Identity;
-	FTransform FinalFkHipTransformCS = FTransform::Identity;
-
+	// TODO: Do adjustments to the ball and hip
 	UE::Anim::FootPlacement::FPlantResult Result =
 	{
 		{ LegData.Bones.IKIndex, CorrectedFootTransformCS },
-		//{ LegData.Bones.BallIndex, FinalFkTipTransformCS },
-		//{ LegData.Bones.HipIndex, FinalFkHipTransformCS }
+		// { LegData.Bones.BallIndex, CorrectedBallTransformCS },
+		//{ LegData.Bones.HipIndex, CorrectedHipTransformCS }
 	};
 
 	return Result;
@@ -1116,9 +1112,9 @@ void FAnimNode_FootPlacement::EvaluateSkeletalControl_AnyThread(FComponentSpaceP
 
 		const UE::Anim::FootPlacement::FPlantResult PlantResult =
 			FinalizeFootAlignment(FootPlacementContext, LegData, LegDef, PelvisTransformCS);
-		OutBoneTransforms.Add(PlantResult.IkPlantTranformCS);
-		//OutBoneTransforms.Add(PlantResult.FkTipTransformCS);
-		//OutBoneTransforms.Add(PlantResult.FkHipTransformCS);
+		OutBoneTransforms.Add(PlantResult.FootTranformCS);
+		//OutBoneTransforms.Add(PlantResult.BallTransformCS);
+		//OutBoneTransforms.Add(PlantResult.HipTransformCS);
 
 #if ENABLE_ANIM_DEBUG
 		if (CVarAnimNodeFootPlacementDebug.GetValueOnAnyThread())
@@ -1127,7 +1123,7 @@ void FAnimNode_FootPlacement::EvaluateSkeletalControl_AnyThread(FComponentSpaceP
 
 			// Grab positions to debug draw history
 			DebugData.OutputFootLocationsWS[FootIndex] =
-				FootPlacementContext.OwningComponentToWorld.TransformPosition(PlantResult.IkPlantTranformCS.Transform.GetLocation());
+				FootPlacementContext.OwningComponentToWorld.TransformPosition(PlantResult.FootTranformCS.Transform.GetLocation());
 			DebugData.InputFootLocationsWS[FootIndex] = 
 				FootPlacementContext.OwningComponentToWorld.TransformPosition(LegData.InputPose.FootTransformCS.GetLocation());;
 		}
@@ -1636,10 +1632,9 @@ FTransform FAnimNode_FootPlacement::SolvePelvis(const UE::Anim::FootPlacement::F
 		FPelvisOffsetRangeForLimb PelvisOffsetRangeCS;
 		FindPelvisOffsetRangeForLimb(
 			Context,
-			LegData.InputPose,
+			LegData,
 			LegData.AlignedFootTransformCS.GetLocation(),
 			PelvisData.InputPose.FKTransformCS,
-			LegData.Bones.LimbLength,
 			PelvisOffsetRangeCS);
 
 		const float DesiredOffset = PelvisOffsetRangeCS.DesiredExtension;
diff --git a/Engine/Plugins/Animation/AnimationWarping/Source/Runtime/Public/BoneControllers/AnimNode_FootPlacement.h b/Engine/Plugins/Animation/AnimationWarping/Source/Runtime/Public/BoneControllers/AnimNode_FootPlacement.h
--- a/Engine/Plugins/Animation/AnimationWarping/Source/Runtime/Public/BoneControllers/AnimNode_FootPlacement.h
+++ b/Engine/Plugins/Animation/AnimationWarping/Source/Runtime/Public/BoneControllers/AnimNode_FootPlacement.h
@@ -133,9 +133,9 @@ namespace UE::Anim::FootPlacement
 	struct FPlantResult
 	{
 	public:
-		FBoneTransform IkPlantTranformCS;
-		//FBoneTransform FkTipTransformCS;
-		//FBoneTransform FkHipTransformCS;
+		FBoneTransform FootTranformCS;
+		//FBoneTransform BallTransformCS;
+		//FBoneTransform HipTransformCS;
 	};
 
 #if ENABLE_ANIM_DEBUG
@@ -427,12 +427,6 @@ public:
 	UPROPERTY(EditAnywhere, Category = "Plant Settings", meta = (ClampMin = "0.0", UIMin = "0.0", ClampMax = "1.0", UIMax = "1.0"))
 	float AnkleTwistReduction = 0.75f;
 
-	// How much we can pull the the foot towards the hip to prevent hyperextension
-	// Before pulling the IK foot towards the FK foot
-	// While we're planted and within this threshold, the foot will roll instead of sliding
-	UPROPERTY(EditAnywhere, Category = "Settings")
-	float ExtensionPlantedPullOffset = 2.0f;
-
 	void Initialize(const FAnimationInitializeContext& Context);
 };
 
@@ -563,10 +557,9 @@ private:
 	// Find the horizontal pelvis offset range for the foot to reach:
 	void FindPelvisOffsetRangeForLimb(
 		const UE::Anim::FootPlacement::FEvaluationContext& Context,
-		const UE::Anim::FootPlacement::FLegRuntimeData::FInputPoseData& LegInputPoseData,
+		const UE::Anim::FootPlacement::FLegRuntimeData& LegData,
 		const FVector& PlantTargetLocationCS,
 		const FTransform& PelvisTransformCS,
-		const float LimbLength,
 		FPelvisOffsetRangeForLimb& OutPelvisOffsetRangeCS) const;
 
 	// Adjust LastPlantTransformWS to current, to have the foot pivot around the ball instead of the ankle
