Foot Placement:
-Rebalance pelvis before height calculation
-Don't lift the pelvis during over-compression check, if the foot is in the airin the source anim
-Added option to reach with toe before planting the foot

diff --git a/Engine/Plugins/Animation/AnimationWarping/Source/Runtime/Private/BoneControllers/AnimNode_FootPlacement.cpp b/Engine/Plugins/Animation/AnimationWarping/Source/Runtime/Private/BoneControllers/AnimNode_FootPlacement.cpp
--- a/Engine/Plugins/Animation/AnimationWarping/Source/Runtime/Private/BoneControllers/AnimNode_FootPlacement.cpp
+++ b/Engine/Plugins/Animation/AnimationWarping/Source/Runtime/Private/BoneControllers/AnimNode_FootPlacement.cpp
@@ -364,7 +364,8 @@ void FAnimNode_FootPlacement::FindPelvisOffsetRangeForLimb(
 	const FVector MinOffsetLocation = DesiredPlantTargetLocationCS + -Context.ApproachDirCS * MinExtension;
 
 	const float MinOffset = (MinOffsetLocation - HipLocationCS) | -Context.ApproachDirCS;
-	OutPelvisOffsetRangeCS.MinExtension = MinOffset;
+	// Limit pelvis compression adjustment by the height of the foot. We can always bring the foot closer to the ground in post-adjustments
+	OutPelvisOffsetRangeCS.MinExtension = MinOffset - LegInputPose.DistanceToPlant;
 }
 
 float FAnimNode_FootPlacement::CalcTargetPlantPlaneDistance(
@@ -780,8 +781,8 @@ UE::Anim::FootPlacement::FPlantResult FAnimNode_FootPlacement::FinalizeFootAlign
 				const bool bRecentlyUnplanted = !bIsPlanted && LegData.Plant.TimeSinceFullyUnaligned == 0.0f;
 				// Try to keep the tip on spot if we're unplanting
 				// Don't do this until we've reached the plant target once
-				// TODO: Make this configurable? 
-				if (bRecentlyUnplanted || (bIsPlanted && LegData.Plant.bCanReachTarget))
+				const bool bCanLiftHeel = bRecentlyUnplanted || (bIsPlanted && LegData.Plant.bCanReachTarget) || PlantSettings.bAdjustHeelBeforePlanting;
+				if (bCanLiftHeel)
 				{
 					// Scale this value by our FK transition alpha to not pop
 					const float MaxPullTowardsHip =  FMath::Min(LegData.Bones.FootLength, HyperExtensionRemaining) * LegData.InputPose.AlignmentAlpha;
@@ -1616,6 +1617,23 @@ FTransform FAnimNode_FootPlacement::SolvePelvis(const UE::Anim::FootPlacement::F
 	TBitArray<> RelevantFeet = FindRelevantFeet(Context);
 	const int32 RelevantFeetNum = RelevantFeet.CountSetBits();
 
+	// Rebalance the pelvis before calculating its desired height
+	FTransform RebalancedPelvisTransform  = PelvisData.InputPose.FKTransformCS;
+	if (PelvisSettings.HorizontalRebalancingWeight)
+	{
+		const int32 NumLegs = LegsData.Num();
+		FVector OffsetAverage = FVector::ZeroVector;
+		for (const FLegRuntimeData& LegData : LegsData)
+		{
+			const FVector LegTranslationOffset = LegData.AlignedFootTransformCS.GetLocation() - LegData.InputPose.FootTransformCS.GetLocation();
+			OffsetAverage += LegTranslationOffset / NumLegs;
+		}
+
+		// Remove the vertical component
+		OffsetAverage += OffsetAverage * Context.ApproachDirCS;
+		RebalancedPelvisTransform.SetLocation(RebalancedPelvisTransform.GetLocation() + OffsetAverage * PelvisSettings.HorizontalRebalancingWeight);
+	}
+
 	// Taken from http://runevision.com/thesis/rune_skovbo_johansen_thesis.pdf
 	// Chapter 7.4.2
 
@@ -1634,7 +1652,7 @@ FTransform FAnimNode_FootPlacement::SolvePelvis(const UE::Anim::FootPlacement::F
 			Context,
 			LegData,
 			LegData.AlignedFootTransformCS.GetLocation(),
-			PelvisData.InputPose.FKTransformCS,
+			RebalancedPelvisTransform,
 			PelvisOffsetRangeCS);
 
 		const float DesiredOffset = PelvisOffsetRangeCS.DesiredExtension;
@@ -1663,22 +1681,6 @@ FTransform FAnimNode_FootPlacement::SolvePelvis(const UE::Anim::FootPlacement::F
 
 	FVector PelvisOffsetDelta = -PelvisOffsetZ * Context.ApproachDirCS;
 
-	if (PelvisSettings.HorizontalRebalancingWeight)
-	{
-		const int32 NumLegs = LegsData.Num();
-		FVector OffsetAverage = FVector::ZeroVector;
-		for (const FLegRuntimeData& LegData : LegsData)
-		{
-			const FVector LegTranslationOffset = LegData.AlignedFootTransformCS.GetLocation() - LegData.InputPose.FootTransformCS.GetLocation();
-			OffsetAverage += LegTranslationOffset / NumLegs;
-		}
-
-		// Remove the vertical component
-		OffsetAverage += OffsetAverage * Context.ApproachDirCS;
-
-		PelvisOffsetDelta += OffsetAverage * PelvisSettings.HorizontalRebalancingWeight;
-	}
-
 	FTransform PelvisTransformCS = PelvisData.InputPose.FKTransformCS;
 	PelvisTransformCS.AddToTranslation(PelvisOffsetDelta);
 
diff --git a/Engine/Plugins/Animation/AnimationWarping/Source/Runtime/Public/BoneControllers/AnimNode_FootPlacement.h b/Engine/Plugins/Animation/AnimationWarping/Source/Runtime/Public/BoneControllers/AnimNode_FootPlacement.h
--- a/Engine/Plugins/Animation/AnimationWarping/Source/Runtime/Public/BoneControllers/AnimNode_FootPlacement.h
+++ b/Engine/Plugins/Animation/AnimationWarping/Source/Runtime/Public/BoneControllers/AnimNode_FootPlacement.h
@@ -427,6 +427,10 @@ public:
 	UPROPERTY(EditAnywhere, Category = "Plant Settings", meta = (ClampMin = "0.0", UIMin = "0.0", ClampMax = "1.0", UIMax = "1.0"))
 	float AnkleTwistReduction = 0.75f;
 
+	// Whether to allow adjusting the heel lift before we plant
+	UPROPERTY(EditAnywhere, Category = "Plant Settings")
+	bool bAdjustHeelBeforePlanting = false;
+
 	void Initialize(const FAnimationInitializeContext& Context);
 };
 
