Foot Placement:
-Fix weight rebalancing
-Change ground spring to be height only
-When traces fail, use capsule's ground
#rnx

diff --git a/Engine/Plugins/Animation/AnimationWarping/Source/Runtime/Private/BoneControllers/AnimNode_FootPlacement.cpp b/Engine/Plugins/Animation/AnimationWarping/Source/Runtime/Private/BoneControllers/AnimNode_FootPlacement.cpp
--- a/Engine/Plugins/Animation/AnimationWarping/Source/Runtime/Private/BoneControllers/AnimNode_FootPlacement.cpp
+++ b/Engine/Plugins/Animation/AnimationWarping/Source/Runtime/Private/BoneControllers/AnimNode_FootPlacement.cpp
@@ -53,6 +53,17 @@ namespace UE::Anim::FootPlacement
 				MovementComponent->CurrentFloor.HitResult.ImpactNormal : -ApproachDirWS;
 		}
 
+		FVector GetMovementComponentFloorLocation() const
+		{
+			if (!MovementComponent)
+			{
+				return OwningComponentToWorld.GetLocation();
+			}
+
+			return (MovementComponent->CurrentFloor.bBlockingHit) ?
+				MovementComponent->CurrentFloor.HitResult.ImpactPoint : OwningComponentToWorld.GetLocation();
+		}
+
 		bool GetMovementComponentIsWalkable(const FHitResult& InHit) const
 		{
 			if (!MovementComponent)
@@ -151,8 +162,8 @@ namespace UE::Anim::FootPlacement
 
 		if (!bHit)
 		{
-			// If the hit fails, use starting position and a default impact normal (negated trace direction)
-			OutImpactLocationWS = StartPositionWS;
+			// If the hit fails, use the ground plane position and a default impact normal (negated trace direction)
+			OutImpactLocationWS = PointDirectionPlaneIntersection(StartPositionWS, TraceDirectionWS, FPlane(Context.GetMovementComponentFloorLocation(), -TraceDirectionWS));
 			OutImpactNormalWS = -TraceDirectionWS;
 			return false;
 		}
@@ -523,7 +534,7 @@ void FAnimNode_FootPlacement::UpdatePlantingPlaneInterpolation(
 
 	if (InterpolationSettings.bEnableFloorInterpolation)
 	{
-		const FVector CurrPlaneIntersection = UE::Anim::FootPlacement::PointDirectionPlaneIntersection(
+		FVector CurrPlaneIntersection = UE::Anim::FootPlacement::PointDirectionPlaneIntersection(
 			FootTransformWS.GetLocation(),
 			TraceDirection,
 			InOutPlantPlane);
@@ -543,26 +554,22 @@ void FAnimNode_FootPlacement::UpdatePlantingPlaneInterpolation(
 		const float AdjustedPrevZ = FMath::Abs(LastPlaneDeltaZ) < FMath::Abs(PrevPlaneDeltaZ) ?
 			LastPlaneIntersection.Z : PrevPlaneIntersection.Z;
 
-		//TODO: replace by Z? Do some math and interpolate Plane.W!
-		const FVector AdjustedPrevPlaneIntersection = FVector(CurrPlaneIntersection.X,
-			CurrPlaneIntersection.Y,
-			AdjustedPrevZ);
-
-		FVector PlantPlaneSpringLocation = UKismetMathLibrary::VectorSpringInterp(
-			AdjustedPrevPlaneIntersection, CurrPlaneIntersection, InOutInterpData.GroundTranslationSpringState,
+		float PlantPlaneSpringHeight = UKismetMathLibrary::FloatSpringInterp(
+			AdjustedPrevZ, CurrPlaneIntersection.Z, InOutInterpData.GroundHeightSpringState,
 			InterpolationSettings.FloorLinearStiffness,
 			InterpolationSettings.FloorLinearDamping,
 			Context.UpdateDeltaTime, 1.0f, 0.0f);
 
+		CurrPlaneIntersection.Z = PlantPlaneSpringHeight;
 		if (TraceSettings.MaxGroundPenetration >= 0.0f)
 		{
 			// Prevent the foot from clipping too much into geometry due to interpolation
 			const float DistanceToGroundPlane =
-				UE::Anim::FootPlacement::GetDistanceToPlaneAlongDirection(PlantPlaneSpringLocation, InOutPlantPlane, Context.ApproachDirCS);
+				UE::Anim::FootPlacement::GetDistanceToPlaneAlongDirection(CurrPlaneIntersection, InOutPlantPlane, Context.ApproachDirCS);
 			const float PenetrationAmount = -DistanceToGroundPlane - TraceSettings.MaxGroundPenetration;
 			if (PenetrationAmount > 0.0f)
 			{
-				PlantPlaneSpringLocation -= Context.ApproachDirCS * PenetrationAmount;
+				CurrPlaneIntersection -= Context.ApproachDirCS * PenetrationAmount;
 			}
 		}
 
@@ -575,7 +582,7 @@ void FAnimNode_FootPlacement::UpdatePlantingPlaneInterpolation(
 			Context.UpdateDeltaTime, 1.0f, 0.0f);
 
 		const FVector PlantPlaneSpringNormal = FloorSpringNormalRotation.RotateVector(LastPlantPlane.GetNormal());
-		const FPlane PlantingPlane = FPlane(PlantPlaneSpringLocation, PlantPlaneSpringNormal);
+		const FPlane PlantingPlane = FPlane(CurrPlaneIntersection, PlantPlaneSpringNormal);
 
 		InOutPlantPlane = PlantingPlane;
 	}
@@ -1619,6 +1626,7 @@ FTransform FAnimNode_FootPlacement::SolvePelvis(const UE::Anim::FootPlacement::F
 
 	// Rebalance the pelvis before calculating its desired height
 	FTransform RebalancedPelvisTransform  = PelvisData.InputPose.FKTransformCS;
+	FVector PelvisOffsetDelta;
 	if (PelvisSettings.HorizontalRebalancingWeight)
 	{
 		const int32 NumLegs = LegsData.Num();
@@ -1630,7 +1638,7 @@ FTransform FAnimNode_FootPlacement::SolvePelvis(const UE::Anim::FootPlacement::F
 		}
 
 		// Remove the vertical component
-		OffsetAverage += OffsetAverage * Context.ApproachDirCS;
+		PelvisOffsetDelta = OffsetAverage - Context.ApproachDirCS.Dot(OffsetAverage) * Context.ApproachDirCS;
 		RebalancedPelvisTransform.SetLocation(RebalancedPelvisTransform.GetLocation() + OffsetAverage * PelvisSettings.HorizontalRebalancingWeight);
 	}
 
@@ -1678,8 +1686,7 @@ FTransform FAnimNode_FootPlacement::SolvePelvis(const UE::Anim::FootPlacement::F
 
 	// Adjust the hips to prevent over-compression
 	PelvisOffsetZ = FMath::Clamp(PelvisOffsetZ, MinOffsetMax, MaxOffsetMin);
-
-	FVector PelvisOffsetDelta = -PelvisOffsetZ * Context.ApproachDirCS;
+	PelvisOffsetDelta += -PelvisOffsetZ * Context.ApproachDirCS;
 
 	FTransform PelvisTransformCS = PelvisData.InputPose.FKTransformCS;
 	PelvisTransformCS.AddToTranslation(PelvisOffsetDelta);
diff --git a/Engine/Plugins/Animation/AnimationWarping/Source/Runtime/Public/BoneControllers/AnimNode_FootPlacement.h b/Engine/Plugins/Animation/AnimationWarping/Source/Runtime/Public/BoneControllers/AnimNode_FootPlacement.h
--- a/Engine/Plugins/Animation/AnimationWarping/Source/Runtime/Public/BoneControllers/AnimNode_FootPlacement.h
+++ b/Engine/Plugins/Animation/AnimationWarping/Source/Runtime/Public/BoneControllers/AnimNode_FootPlacement.h
@@ -78,7 +78,7 @@ namespace UE::Anim::FootPlacement
 			FVectorSpringState PlantOffsetTranslationSpringState;
 			FQuaternionSpringState PlantOffsetRotationSpringState;
 
-			FVectorSpringState GroundTranslationSpringState;
+			FFloatSpringState GroundHeightSpringState;
 			FQuaternionSpringState GroundRotationSpringState;
 		} Interpolation;
 	};
