Foot placement:
-Added "disable foot lock" per-foot alpha
-Improved comments
-Removed unused variables
-Added heel life pelvis ratio
#rnx

diff --git a/Engine/Plugins/Animation/AnimationWarping/Source/Runtime/Private/BoneControllers/AnimNode_FootPlacement.cpp b/Engine/Plugins/Animation/AnimationWarping/Source/Runtime/Private/BoneControllers/AnimNode_FootPlacement.cpp
--- a/Engine/Plugins/Animation/AnimationWarping/Source/Runtime/Private/BoneControllers/AnimNode_FootPlacement.cpp
+++ b/Engine/Plugins/Animation/AnimationWarping/Source/Runtime/Private/BoneControllers/AnimNode_FootPlacement.cpp
@@ -316,7 +316,7 @@ void FAnimNode_FootPlacement::FindPelvisOffsetRangeForLimb(
 		const FVector ToFKFoot = (FKFootProjected - PlantTargetLocationCS);
 		const FVector ToFKFootDir = ToFKFoot.GetSafeNormal();
 		const float ToFKDistance = ToFKFoot.Size();
-		PlantTargetLocationCS = PlantTargetLocationCS + ToFKFootDir * FMath::Min(ToFKDistance, Bones.FootLength);
+		PlantTargetLocationCS = PlantTargetLocationCS + ToFKFootDir * FMath::Min(ToFKDistance, Bones.FootLength) * PelvisSettings.HeelLiftRatio;
 
 		// Calculate new correction offset and direction
 		TargetFootToHip = HipProjected - PlantTargetLocationCS;
@@ -704,7 +704,7 @@ bool FAnimNode_FootPlacement::WantsToPlant(
 	}
 #endif
 
-	if (PlantSettings.LockType == EFootPlacementLockType::Unlocked)
+	if ((PlantSettings.LockType == EFootPlacementLockType::Unlocked) || FMath::IsNearlyZero(LegInputPose.LockAlpha))
 	{
 		return false;
 	}
@@ -1076,8 +1076,6 @@ void FAnimNode_FootPlacement::EvaluateSkeletalControl_AnyThread(FComponentSpaceP
 
 		// TODO: All of these can be calculated on initialize, but in case there's value in changing these dynamically,
 		// will keep this for now. If needed, change to lazy update.
-		PlantRuntimeSettings.MaxExtensionRatioSqrd = PlantSettings.MaxExtensionRatio * PlantSettings.MaxExtensionRatio;
-		PlantRuntimeSettings.MinExtensionRatioSqrd = PlantSettings.MinExtensionRatio * PlantSettings.MinExtensionRatio;
 		PlantRuntimeSettings.UnplantRadiusSqrd = PlantSettings.UnplantRadius * PlantSettings.UnplantRadius;
 		PlantRuntimeSettings.ReplantRadiusSqrd =
 			PlantRuntimeSettings.UnplantRadiusSqrd *
@@ -1243,6 +1241,7 @@ void FAnimNode_FootPlacement::InitializeBoneReferences(const FBoneContainer& Req
 		const FTransform BallTransformLS = RequiredBones.GetRefPoseTransform(LegData.Bones.BallIndex);
 		LegData.Bones.FootLength = BallTransformLS.GetLocation().Size();
 
+#if ENABLE_ANIM_DEBUG
 		// TODO: This wont work for animations authored for different slopes or stairs. Figure this out later
 		const FVector RefPoseGroundNormalCS = FVector::UpVector;
 		const FTransform BallRefTransformCS = FAnimationRuntime::GetComponentSpaceRefPose(
@@ -1258,15 +1257,13 @@ void FAnimNode_FootPlacement::InitializeBoneReferences(const FBoneContainer& Req
 		const FVector FootAlignmentDeltaCS = -FKFootTransformCS.GetLocation();
 		const FVector FootAlignmentOffsetCS = (FootAlignmentDeltaCS | RefPoseGroundNormalCS) * RefPoseGroundNormalCS;
 		LegData.InputPose.FootToGround = FTransform(FKFootTransformCS.GetRotation().UnrotateVector(FootAlignmentOffsetCS));
+#endif 
 
 		const USkeleton* Skeleton = RequiredBones.GetSkeletonAsset();
 		check(Skeleton);
-		SmartName::UID_Type NameUID = Skeleton->GetUIDByName(USkeleton::AnimCurveMappingName, LegDef.SpeedCurveName);
-		if (NameUID != SmartName::MaxUID)
-		{
-			// Grab UIDs of filtered curves to avoid lookup later
-			LegData.SpeedCurveUID = NameUID;
-		}
+		// Grab UIDs of filtered curves to avoid lookup later
+		LegData.SpeedCurveUID = Skeleton->GetUIDByName(USkeleton::AnimCurveMappingName, LegDef.SpeedCurveName);
+		LegData.DisableLockCurveUID = Skeleton->GetUIDByName(USkeleton::AnimCurveMappingName, LegDef.DisableLockCurveName);
 	}
 
 	PelvisBone.Initialize(RequiredBones);
@@ -1288,7 +1285,6 @@ void FAnimNode_FootPlacement::GatherPelvisDataFromInputs(const UE::Anim::FootPla
 	// TODO: All of these can be calculated on initialize, but in case there's value in changing these dynamically,
 	// will keep this for now. If needed, change to lazy update.
 	PelvisData.MaxOffsetSqrd = PelvisSettings.MaxOffset * PelvisSettings.MaxOffset;
-	PelvisData.MaxOffsetHorizontalSqrd = PelvisSettings.MaxOffsetHorizontal * PelvisSettings.MaxOffsetHorizontal;
 }
 
 void FAnimNode_FootPlacement::GatherLegDataFromInputs(
@@ -1356,6 +1352,9 @@ void FAnimNode_FootPlacement::GatherLegDataFromInputs(
 			Context.CSPContext.Curve.Get(LegData.SpeedCurveUID, bValidSpeedCurve, DefaultSpeedCurveValue);
 	}
 
+	// Grab the lock curve's alpha. If the curve isn't set, then LockAlpha is full weight.
+	LegData.InputPose.LockAlpha = 1.0f - Context.CSPContext.Curve.Get(LegData.DisableLockCurveUID);
+
 	LegData.InputPose.DistanceToPlant = CalcTargetPlantPlaneDistance(Context, LegData.InputPose);
 	const float FKAlignmentAlpha = GetAlignmentAlpha(Context, LegData.InputPose);
 	LegData.InputPose.AlignmentAlpha = FKAlignmentAlpha;
@@ -1578,7 +1577,14 @@ void FAnimNode_FootPlacement::ProcessFootAlignment(
 #endif
 	}
 
-	LegData.UnalignedFootTransformWS = FootUnalignedTransformCS * Context.OwningComponentToWorld;
+	FTransform BlendedUnalignedTransformCS;
+	{
+		// Blend the component-space input pose, with the unaligned foot-locked transform.
+		// Allow ground alignment to continue with this blended result.
+		// When the lock alpha reaches 0, we will automatically unlock the foot.
+		BlendedUnalignedTransformCS.Blend(InputPose.FootTransformCS, FootUnalignedTransformCS, InputPose.LockAlpha);
+		LegData.UnalignedFootTransformWS = BlendedUnalignedTransformCS * Context.OwningComponentToWorld;
+	}
 
 	const FTransform ComponentToWorldInv = Context.OwningComponentToWorld.Inverse();
 
@@ -1638,8 +1644,8 @@ FTransform FAnimNode_FootPlacement::SolvePelvis(const UE::Anim::FootPlacement::F
 		}
 
 		// Remove the vertical component
-		PelvisOffsetDelta = OffsetAverage - Context.ApproachDirCS.Dot(OffsetAverage) * Context.ApproachDirCS;
-		RebalancedPelvisTransform.SetLocation(RebalancedPelvisTransform.GetLocation() + OffsetAverage * PelvisSettings.HorizontalRebalancingWeight);
+		PelvisOffsetDelta = (OffsetAverage - Context.ApproachDirCS.Dot(OffsetAverage) * Context.ApproachDirCS) * PelvisSettings.HorizontalRebalancingWeight;
+		RebalancedPelvisTransform.SetLocation(RebalancedPelvisTransform.GetLocation() + PelvisOffsetDelta);
 	}
 
 	// Taken from http://runevision.com/thesis/rune_skovbo_johansen_thesis.pdf
diff --git a/Engine/Plugins/Animation/AnimationWarping/Source/Runtime/Public/BoneControllers/AnimNode_FootPlacement.h b/Engine/Plugins/Animation/AnimationWarping/Source/Runtime/Public/BoneControllers/AnimNode_FootPlacement.h
--- a/Engine/Plugins/Animation/AnimationWarping/Source/Runtime/Public/BoneControllers/AnimNode_FootPlacement.h
+++ b/Engine/Plugins/Animation/AnimationWarping/Source/Runtime/Public/BoneControllers/AnimNode_FootPlacement.h
@@ -23,6 +23,7 @@ namespace UE::Anim::FootPlacement
 	{
 		int32 Idx = -1;
 
+		// Bone information that can be cached once per-lod change.
 		struct FBoneData
 		{
 			FCompactPoseBoneIndex FKIndex = FCompactPoseBoneIndex(INDEX_NONE);
@@ -33,9 +34,11 @@ namespace UE::Anim::FootPlacement
 			float FootLength = 0.0f;
 		} Bones;
 
+		// Curves
 		SmartName::UID_Type SpeedCurveUID = SmartName::MaxUID;
+		SmartName::UID_Type DisableLockCurveUID = SmartName::MaxUID;
 
-		// Helper struct to store values used across the evaluation
+		// Helper struct to store values coming directly, or trivial to calculate from just the input pose.
 		struct FInputPoseData
 		{
 			FTransform FootTransformCS = FTransform::Identity;
@@ -43,10 +46,17 @@ namespace UE::Anim::FootPlacement
 			FTransform HipTransformCS = FTransform::Identity;
 			FTransform BallToFoot = FTransform::Identity;
 			FTransform FootToBall = FTransform::Identity;
+#if ENABLE_ANIM_DEBUG
+			// These are only used for debug draw at the moment
+			// @TODO: Use this info to more precisely figure out foot dimensions
 			FTransform FootToGround = FTransform::Identity;
 			FTransform BallToGround = FTransform::Identity;
+#endif
 			float Speed = 0.0f;
+			float LockAlpha = 0.0f;
 			float DistanceToPlant = 0.0f;
+			// Calculated from a range of toe speeds to define when to blend in/out ground rotational alignment
+			// @TODO: When we have prediction/phase info, replac with roll-phase alpha
 			float AlignmentAlpha = 0.0f;
 		} InputPose;
 
@@ -59,25 +69,29 @@ namespace UE::Anim::FootPlacement
 			FPlane PlantPlaneWS = FPlane(FVector::UpVector, 0.0f);
 			FPlane PlantPlaneCS = FPlane(FVector::UpVector, 0.0f);
 			FQuat TwistCorrection = FQuat::Identity;
+			// @TODO: When we have prediction/phase info, replace use-cases with post-plant roll-phase
 			float TimeSinceFullyUnaligned = 0.0f;
+			// Whether the planted/locked target has ever been reachable this plant
 			bool bCanReachTarget = false;
 			// Whether we want to plant, independently from any dynamic pose adjustments we may do
 			bool bWantsToPlant = false;
 		} Plant;
 		
+		// Ground-aligned, locked/unlocked bone transform pre-extension adjustments
 		FTransform AlignedFootTransformWS = FTransform::Identity;
-		FTransform UnalignedFootTransformWS = FTransform::Identity;
 		FTransform AlignedFootTransformCS = FTransform::Identity;
-		FVector CachedIKToFKDir = FVector::UpVector;
+		// Foot locked/unlocked bone transform, before ground alignment
+		FTransform UnalignedFootTransformWS = FTransform::Identity;
 		
 		/* Interpolation */
 		struct FInterpolationData
 		{
+			// Interpolated foot lock offset
 			FTransform UnalignedFootOffsetCS = FTransform::Identity;
-
+			// Foot lock pring states
 			FVectorSpringState PlantOffsetTranslationSpringState;
 			FQuaternionSpringState PlantOffsetRotationSpringState;
-
+			// Ground alignment spring states
 			FFloatSpringState GroundHeightSpringState;
 			FQuaternionSpringState GroundRotationSpringState;
 		} Interpolation;
@@ -85,8 +99,6 @@ namespace UE::Anim::FootPlacement
 
 	struct FPlantRuntimeSettings
 	{
-		float MaxExtensionRatioSqrd = 0.0f;
-		float MinExtensionRatioSqrd = 0.0f;
 		float UnplantRadiusSqrd = 0.0f;
 		float ReplantRadiusSqrd = 0.0f;
 		float CosHalfUnplantAngle = 0.0f;
@@ -103,7 +115,6 @@ namespace UE::Anim::FootPlacement
 		} Bones;
 
 		/* Settings-based properties */
-		float MaxOffsetHorizontalSqrd = 0.0f;
 		float MaxOffsetSqrd = 0.0f;
 
 		/* Input pose properties */
@@ -118,6 +129,7 @@ namespace UE::Anim::FootPlacement
 		/* Interpolation */
 		struct FInterpolationData
 		{
+			// Current pelvis offset and spring states. We use a 3d vector because this interpolates weight rebalancing too.
 			FVector PelvisTranslationOffset = FVector::ZeroVector;
 			FVectorSpringState PelvisTranslationSpringState;
 		} Interpolation;
@@ -130,11 +142,14 @@ namespace UE::Anim::FootPlacement
 		bool bIsOnGround = false;
 	};
 
+	// Final result after post-adjustments (extension checks, heel lift, etc.)
 	struct FPlantResult
 	{
 	public:
 		FBoneTransform FootTranformCS;
+		// @TODO: Add procedural toe rolling
 		//FBoneTransform BallTransformCS;
+		// @TODO: Look into shifting/rotating the hips to prevent over-extension, and give IK an easier time.
 		//FBoneTransform HipTransformCS;
 	};
 
@@ -171,10 +186,15 @@ namespace UE::Anim::FootPlacement
 UENUM(BlueprintType)
 enum class EFootPlacementLockType : uint8
 {
+	// Foot is unlocked but free to move
 	Unlocked,
+	// Foot can lock, and will pivot around its ball/toes.
 	PivotAroundBall,
+	// Foot can lock, and will pivot around the ankle/foot bone.
 	PivotAroundAnkle,
+	// Foot is fully locked. Useful for bigger/mechanical creatures
 	LockRotation
+	// @TODO: Detect whether the ball or ankle is planted. Shift the rotation pivot-style depending on this.
 };
 
 USTRUCT(BlueprintType)
@@ -328,6 +348,12 @@ public:
 	UPROPERTY(EditAnywhere, Category = "Pelvis Settings", meta=(DisplayAfter="bEnableInterpolation"))
 	float MaxOffsetHorizontal = 10.0f;
 
+	// How much we prefer lifting the heel before dropping the hips to achieve the desired pose.
+	// 1 fully favors heel lift first.
+	// 0 fully favors pelvis drop first.
+	UPROPERTY(EditAnywhere, Category = "Pelvis Settings", meta = (ClampMin = "0.0", UIMin = "0.0", ClampMax = "1.0", UIMax = "1.0", DisplayAfter="bEnableInterpolation"))
+	float HeelLiftRatio = 0.5f;
+
 	UPROPERTY(EditAnywhere, Category = "Pelvis Settings", meta=(DisplayAfter="bEnableInterpolation"))
 	EPelvisHeightMode PelvisHeightMode = EPelvisHeightMode::AllLegs;
 
@@ -365,6 +391,11 @@ public:
 	UPROPERTY(EditAnywhere, Category = "Settings")
 	FName SpeedCurveName = NAME_None;
 
+	// Name of the curve representing the alpha of the locking alpha.
+	// This allows you to disable locking precisely, instead of relying on the procedural mechanism based on springs and foot analysis
+	UPROPERTY(EditAnywhere, Category = "Settings")
+	FName DisableLockCurveName = NAME_None;
+
 public:
 
 	void InitializeBoneReferences(const FBoneContainer& RequiredBones);
@@ -448,7 +479,6 @@ public:
 	UPROPERTY(EditAnywhere, Category = "Settings")
 	EWarpingEvaluationMode PlantSpeedMode = EWarpingEvaluationMode::Manual;
 
-	// TODO: This wont work well when the animation doesn't have a single plant plane, i.e. a walking upstairs anim
 	UPROPERTY(EditAnywhere, Category = "Settings")
 	FBoneReference IKFootRootBone;
 
@@ -491,26 +521,28 @@ private:
 	virtual void InitializeBoneReferences(const FBoneContainer& RequiredBones) override;
 	// End of FAnimNode_SkeletalControlBase interface
 
+	// Gather raw or trivially calculated values from input pose
 	void GatherPelvisDataFromInputs(const UE::Anim::FootPlacement::FEvaluationContext& Context);
 	void GatherLegDataFromInputs(
 		const UE::Anim::FootPlacement::FEvaluationContext& Context,
 		UE::Anim::FootPlacement::FLegRuntimeData& LegData,
 		const FFootPlacemenLegDefinition& LegDef);
+
 	void CalculateFootMidpoint(
 		const UE::Anim::FootPlacement::FEvaluationContext& Context,
 		TConstArrayView<UE::Anim::FootPlacement::FLegRuntimeData> LegData,
 		FVector& OutMidpoint) const;
+
+	// Calculate procedural adjustments before solving the desired pelvis position
 	void ProcessCharacterState(const UE::Anim::FootPlacement::FEvaluationContext& Context);
 	void ProcessFootAlignment(
 		const UE::Anim::FootPlacement::FEvaluationContext& Context,
 		UE::Anim::FootPlacement::FLegRuntimeData& LegData);
 
-	// Evaluates the Pelvis offsets required by the different planting bones and produces a result that best
-	// acommodates all of them
+	// Calculate the desired pelvis offset, based on procedural character/foot adjustments
 	FTransform SolvePelvis(const UE::Anim::FootPlacement::FEvaluationContext& Context);
 	TBitArray<> FindRelevantFeet(const UE::Anim::FootPlacement::FEvaluationContext& Context);
 	TBitArray<> FindPlantedFeet(const UE::Anim::FootPlacement::FEvaluationContext& Context);
-
 	FTransform UpdatePelvisInterpolation(
 		const UE::Anim::FootPlacement::FEvaluationContext& Context,
 		const FTransform& TargetPelvisTransform);
