Fix for orientation warping jitter on listen servers.

diff --git a/Engine/Plugins/Animation/AnimationWarping/Source/Runtime/Private/BoneControllers/AnimNode_OrientationWarping.cpp b/Engine/Plugins/Animation/AnimationWarping/Source/Runtime/Private/BoneControllers/AnimNode_OrientationWarping.cpp
--- a/Engine/Plugins/Animation/AnimationWarping/Source/Runtime/Private/BoneControllers/AnimNode_OrientationWarping.cpp
+++ b/Engine/Plugins/Animation/AnimationWarping/Source/Runtime/Private/BoneControllers/AnimNode_OrientationWarping.cpp
@@ -53,7 +53,8 @@ void FAnimNode_OrientationWarping::GatherDebugData(FNodeDebugData& DebugData)
 		{
 			DebugLine += TEXT("\n - Evaluation Mode: (Graph)");
 			DebugLine += FString::Printf(TEXT("\n - Orientation Angle: (%.3fd)"), FMath::RadiansToDegrees(ActualOrientationAngle));
-			DebugLine += FString::Printf(TEXT("\n - Locomotion Angle: (%.3fd)"), FMath::RadiansToDegrees(LocomotionAngle));
+			// Locomotion angle is already in degrees.
+			DebugLine += FString::Printf(TEXT("\n - Locomotion Angle: (%.3fd)"), LocomotionAngle);
 			DebugLine += FString::Printf(TEXT("\n - Locomotion Delta Angle Threshold: (%.3fd)"), LocomotionAngleDeltaThreshold);
 #if WITH_EDITORONLY_DATA
 			DebugLine += FString::Printf(TEXT("\n - Root Motion Delta Attribute Found: %s)"), (bFoundRootMotionAttribute) ? TEXT("true") : TEXT("false"));
@@ -163,8 +164,9 @@ void FAnimNode_OrientationWarping::EvaluateSkeletalControl_AnyThread(FComponentS
 			// 3. Skeletal Mesh Relative Rotation
 
 			LocomotionAngle = FRotator::NormalizeAxis(LocomotionAngle);
-			LocomotionAngle = FMath::DegreesToRadians(LocomotionAngle);
-			const FQuat LocomotionRotation = FQuat(RotationAxisVector, LocomotionAngle);
+			// UE-184297 Avoid storing LocomotionAngle in radians in case haven't updated the pinned input, to avoid a DegToRad(RadianValue)
+			const float LocomotionAngleRadians = FMath::DegreesToRadians(LocomotionAngle);
+			const FQuat LocomotionRotation = FQuat(RotationAxisVector, LocomotionAngleRadians);
 
 			const FTransform SkeletalMeshRelativeTransform = Output.AnimInstanceProxy->GetComponentRelativeTransform();
 			const FQuat SkeletalMeshRelativeRotation = SkeletalMeshRelativeTransform.GetRotation();
