Orientation Warping:
-Fix unstable behavior on sudden direction changes.
-Allow skipping interpolation on first frame.
-Added parameter to clamp max interpolation in a frame.
-Removed experimental root offset.
-Added on-screen world space debug text.
-Renamed variables so it's obvious which angles are in degrees vs radians.

#rnx

diff --git a/Engine/Plugins/Animation/AnimationWarping/Source/Runtime/Private/BoneControllers/AnimNode_OrientationWarping.cpp b/Engine/Plugins/Animation/AnimationWarping/Source/Runtime/Private/BoneControllers/AnimNode_OrientationWarping.cpp
--- a/Engine/Plugins/Animation/AnimationWarping/Source/Runtime/Private/BoneControllers/AnimNode_OrientationWarping.cpp
+++ b/Engine/Plugins/Animation/AnimationWarping/Source/Runtime/Private/BoneControllers/AnimNode_OrientationWarping.cpp
@@ -11,7 +11,7 @@
 DECLARE_CYCLE_STAT(TEXT("OrientationWarping Eval"), STAT_OrientationWarping_Eval, STATGROUP_Anim);
 
 #if ENABLE_ANIM_DEBUG
-static TAutoConsoleVariable<int32> CVarAnimNodeOrientationWarpingDebug(TEXT("a.AnimNode.OrientationWarping.Debug"), 0, TEXT("Turn on visualization debugging for Orientation Warping. 0 == off, 1 == Warping, 2 == Offset, 3 == Both."));
+static TAutoConsoleVariable<int32> CVarAnimNodeOrientationWarpingDebug(TEXT("a.AnimNode.OrientationWarping.Debug"), 0, TEXT("Turn on visualization debugging for Orientation Warping."));
 static TAutoConsoleVariable<int32> CVarAnimNodeOrientationWarpingVerbose(TEXT("a.AnimNode.OrientationWarping.Verbose"), 0, TEXT("Turn on verbose graph debugging for Orientation Warping"));
 static TAutoConsoleVariable<int32> CVarAnimNodeOrientationWarpingEnable(TEXT("a.AnimNode.OrientationWarping.Enable"), 1, TEXT("Toggle Orientation Warping"));
 #endif
@@ -36,6 +36,15 @@ namespace UE::Anim
 		Angle = FRotator::NormalizeAxis(Angle);
 		return FMath::IsNearlyZero(Angle, Tolerance) || FMath::IsNearlyEqual(FMath::Abs(Angle), 180.f, Tolerance);
 	}
+
+	static float SignedAngleRadBetweenNormals(const FVector& From, const FVector& To, const FVector& Axis)
+	{
+		const float FromDotTo = FVector::DotProduct(From, To);
+		const float Angle = FMath::Acos(FromDotTo);
+		const FVector Cross = FVector::CrossProduct(From, To);
+		const float Dot = FVector::DotProduct(Cross, Axis);
+		return Dot >= 0 ? Angle : -Angle;
+	}
 }
 
 void FAnimNode_OrientationWarping::GatherDebugData(FNodeDebugData& DebugData)
@@ -47,12 +56,12 @@ void FAnimNode_OrientationWarping::GatherDebugData(FNodeDebugData& DebugData)
 		if (Mode == EWarpingEvaluationMode::Manual)
 		{
 			DebugLine += TEXT("\n - Evaluation Mode: (Manual)");
-			DebugLine += FString::Printf(TEXT("\n - Orientation Angle: (%.3fd)"), FMath::RadiansToDegrees(ActualOrientationAngle));
+			DebugLine += FString::Printf(TEXT("\n - Orientation Angle: (%.3fd)"), FMath::RadiansToDegrees(ActualOrientationAngleRad));
 		}
 		else
 		{
 			DebugLine += TEXT("\n - Evaluation Mode: (Graph)");
-			DebugLine += FString::Printf(TEXT("\n - Orientation Angle: (%.3fd)"), FMath::RadiansToDegrees(ActualOrientationAngle));
+			DebugLine += FString::Printf(TEXT("\n - Orientation Angle: (%.3fd)"), FMath::RadiansToDegrees(ActualOrientationAngleRad));
 			// Locomotion angle is already in degrees.
 			DebugLine += FString::Printf(TEXT("\n - Locomotion Angle: (%.3fd)"), LocomotionAngle);
 			DebugLine += FString::Printf(TEXT("\n - Locomotion Delta Angle Threshold: (%.3fd)"), LocomotionAngleDeltaThreshold);
@@ -60,9 +69,6 @@ void FAnimNode_OrientationWarping::GatherDebugData(FNodeDebugData& DebugData)
 			DebugLine += FString::Printf(TEXT("\n - Root Motion Delta Attribute Found: %s)"), (bFoundRootMotionAttribute) ? TEXT("true") : TEXT("false"));
 #endif
 		}
-		DebugLine += FString::Printf(TEXT("\n - Warping Alpha: (%.2fd)"), WarpingAlpha);
-		DebugLine += FString::Printf(TEXT("\n - Heading Offset: (%.3fd)"), FMath::RadiansToDegrees(HeadingOffset));
-		DebugLine += FString::Printf(TEXT("\n - Offset Alpha: (%.2fd)"), OffsetAlpha);
 		if (const UEnum* TypeEnum = FindObject<UEnum>(nullptr, TEXT("/Script/CoreUObject.EAxis")))
 		{
 			DebugLine += FString::Printf(TEXT("\n - Rotation Axis: (%s)"), *(TypeEnum->GetNameStringByIndex(static_cast<int32>(RotationAxis))));
@@ -72,9 +78,8 @@ void FAnimNode_OrientationWarping::GatherDebugData(FNodeDebugData& DebugData)
 	else
 #endif
 	{
-	const float ActualOrientationAngleDegrees = FMath::RadiansToDegrees(ActualOrientationAngle);
-	const float HeadingOffsetDegrees = FMath::RadiansToDegrees(HeadingOffset);
-	DebugLine += FString::Printf(TEXT("(Orientation Angle: %.3fd, Heading Offset : %.3fd, )"), ActualOrientationAngleDegrees, HeadingOffsetDegrees);
+	const float ActualOrientationAngleDegrees = FMath::RadiansToDegrees(ActualOrientationAngleRad);
+	DebugLine += FString::Printf(TEXT("(Orientation Angle: %.3fd)"), ActualOrientationAngleDegrees);
 	}
 	DebugData.AddDebugItem(DebugLine);
 	ComponentPose.GatherDebugData(DebugData);
@@ -84,9 +89,6 @@ void FAnimNode_OrientationWarping::Initialize_AnyThread(const FAnimationInitiali
 {
 	FAnimNode_SkeletalControlBase::Initialize_AnyThread(Context);
 
-	PreviousRootMotionDeltaDirection = FVector::ZeroVector;
-	PreviousOrientationAngle = 0.f;
-	ActualOrientationAngle = 0.f;
 	Reset(Context);
 }
 
@@ -107,11 +109,11 @@ void FAnimNode_OrientationWarping::EvaluateSkeletalControl_AnyThread(FComponentS
 	SCOPE_CYCLE_COUNTER(STAT_OrientationWarping_Eval);
 	check(OutBoneTransforms.Num() == 0);
 
-	ActualOrientationAngle = OrientationAngle;
+	float TargetOrientationAngleRad;
 
 	const float DeltaSeconds = Output.AnimInstanceProxy->GetDeltaSeconds();
+	const float MaxAngleCorrectionRad = FMath::DegreesToRadians(MaxCorrectionRateDegrees) * DeltaSeconds;
 	const FVector RotationAxisVector = UE::Anim::GetAxisVector(RotationAxis);
-	FVector RootMotionDeltaDirection = FVector::ZeroVector;
 	FVector LocomotionForward = FVector::ZeroVector;
 
 	bool bGraphDrivenWarping = false;
@@ -127,6 +129,12 @@ void FAnimNode_OrientationWarping::EvaluateSkeletalControl_AnyThread(FComponentS
 	bFoundRootMotionAttribute = false;
 #endif
 
+#if ENABLE_ANIM_DEBUG
+	FTransform RootMotionTransformDelta = FTransform::Identity;
+	float RootMotionDeltaAngleRad = 0.0;
+	const float PreviousOrientationAngleRad = ActualOrientationAngleRad;
+#endif
+
 	// We will likely need to revisit LocomotionAngle participating as an input to orientation warping.
 	// Without velocity information from the motion model (such as the capsule), LocomotionAngle isn't enough
 	// information in isolation for all cases when deciding to warp.
@@ -143,7 +151,10 @@ void FAnimNode_OrientationWarping::EvaluateSkeletalControl_AnyThread(FComponentS
 
 	if (bGraphDrivenWarping)
 	{
+#if !ENABLE_ANIM_DEBUG
 		FTransform RootMotionTransformDelta = FTransform::Identity;
+#endif
+
 		bGraphDrivenWarping = RootMotionProvider->ExtractRootMotion(Output.CustomAttributes, RootMotionTransformDelta);
 
 		// Graph driven orientation warping will modify the incoming root motion to orient towards the intended locomotion angle
@@ -171,90 +182,55 @@ void FAnimNode_OrientationWarping::EvaluateSkeletalControl_AnyThread(FComponentS
 			const FTransform SkeletalMeshRelativeTransform = Output.AnimInstanceProxy->GetComponentRelativeTransform();
 			const FQuat SkeletalMeshRelativeRotation = SkeletalMeshRelativeTransform.GetRotation();
 			LocomotionForward = SkeletalMeshRelativeRotation.UnrotateVector(LocomotionRotation.GetForwardVector()).GetSafeNormal();
-			
-			if (OffsetAlpha == 0.0)
-			{
-				// No alpha. Reset our heading offset.
-				HeadingOffset = 0.0f;
-			}
-			else
-			{
-				// Accumulate a percentage of our component's frame delta, based on offset alpha
-				const float LastComponentHeading = ComponentHeading;
-				ComponentHeading = Output.AnimInstanceProxy->GetComponentTransform().GetRotation().GetTwistAngle(RotationAxisVector);
-				HeadingOffset += FMath::UnwindRadians(LastComponentHeading - ComponentHeading) * OffsetAlpha;
-
-				// Accumulate a percentage of our root motion's heading, based on offset alpha
-				float RootMotionDeltaHeading = RootMotionTransformDelta.GetRotation().GetTwistAngle(RotationAxisVector);
-				HeadingOffset += RootMotionDeltaHeading * OffsetAlpha;
-
-				// If our alpha is decreasing, we're blending out the offset.
-				// We don't need to do this when blending in because we're just accumulating partial component offset and root motion.
-				const float OffsetAlphaStep = LastOffsetAlpha - OffsetAlpha;
-				if (OffsetAlphaStep > 0.0f)
-				{
-					HeadingOffset -= OffsetAlphaStep * HeadingOffset / LastOffsetAlpha;
-				}
-
-				const float MaxOffsetRadians = FMath::DegreesToRadians(MaxOffsetAngle);
-				HeadingOffset = FMath::Clamp(HeadingOffset, -MaxOffsetAngle, MaxOffsetAngle);
-			}
-			LastOffsetAlpha = OffsetAlpha;
-
-			// Rotate the root motion direction by our effective offset.
-			// This means we will warp our pose based on the adjusted offset, and correct accordingly.
-			FQuat RootOffsetRotation = FQuat(RotationAxisVector, HeadingOffset);
-			const FVector RootMotionDeltaTranslation = RootOffsetRotation.RotateVector(RootMotionTransformDelta.GetTranslation());
 
-			// Hold previous direction if we can't calculate it from current move delta, because the root is no longer moving
-			RootMotionDeltaDirection = RootMotionDeltaTranslation.GetSafeNormal(UE_SMALL_NUMBER, PreviousRootMotionDeltaDirection);
+			const FVector RootMotionDeltaTranslation = RootMotionTransformDelta.GetTranslation();
 
 			const float RootMotionDeltaSpeed = RootMotionDeltaTranslation.Size() / DeltaSeconds;
 			if (RootMotionDeltaSpeed < MinRootMotionSpeedThreshold)
 			{
 				// If we're under the threshold, snap orientation angle to 0, and let interpolation handle the delta
-				ActualOrientationAngle = 0.0f;
-				PreviousRootMotionDeltaDirection = RootMotionDeltaDirection;
+				TargetOrientationAngleRad = 0.0f;
 			}
 			else
 			{
-				// Capture the delta rotation from the axis of motion we care about
-				FQuat WarpedRotation = FQuat::FindBetween(RootMotionDeltaDirection, LocomotionForward);
-
-				// For interpolated warping, guarantee that PreviousOrientationAngle is relative to the current frame's root motion direction 
-				float RootMotionDeltaAngleDifference = FMath::Acos(RootMotionDeltaDirection.Dot(PreviousRootMotionDeltaDirection));
-				RootMotionDeltaAngleDifference *= FMath::Sign(RotationAxisVector.Dot(RootMotionDeltaDirection.Cross(PreviousRootMotionDeltaDirection)));
+				const FVector PreviousRootMotionDeltaDirection = RootMotionDeltaDirection;
+				// Hold previous direction if we can't calculate it from current move delta, because the root is no longer moving
+				RootMotionDeltaDirection = RootMotionDeltaTranslation.GetSafeNormal(UE_SMALL_NUMBER, PreviousRootMotionDeltaDirection);
+				TargetOrientationAngleRad = UE::Anim::SignedAngleRadBetweenNormals(RootMotionDeltaDirection, LocomotionForward, RotationAxisVector);
 
-				PreviousRootMotionDeltaDirection = RootMotionDeltaDirection;
-				PreviousOrientationAngle += RootMotionDeltaAngleDifference;
-
-				ActualOrientationAngle = WarpedRotation.GetTwistAngle(RotationAxisVector);
 				// Motion Matching may return an animation that deviates a lot from the movement direction (e.g movement direction going bwd and motion matching could return the fwd animation for a few frames)
 				// When that happens, since we use the delta between root motion and movement direction, we would be over-rotating the lower body and breaking the pose during those frames
-				// So, when that happens we use the inverse of the movement direction to calculate our target rotation. 
+				// So, when that happens we use the inverse of the root motion direction to calculate our target rotation. 
 				// This feels a bit 'hacky' but its the only option I've found so far to mitigate the problem
 				if (LocomotionAngleDeltaThreshold > 0.f)
 				{
-					if (FMath::Abs(FMath::RadiansToDegrees(ActualOrientationAngle)) > LocomotionAngleDeltaThreshold)
-					{
-						WarpedRotation = FQuat::FindBetween(RootMotionDeltaDirection, -LocomotionForward);
-						ActualOrientationAngle = WarpedRotation.GetTwistAngle(RotationAxisVector);
-					}
-					
-					if (FMath::Abs(FMath::RadiansToDegrees(PreviousOrientationAngle)) > LocomotionAngleDeltaThreshold)
+					if (FMath::Abs(FMath::RadiansToDegrees(TargetOrientationAngleRad)) > LocomotionAngleDeltaThreshold)
 					{
-						// Previous orientation angle might be using an opposite direction too, so flip it if it exceeds the threshold as well.
-						PreviousOrientationAngle = WarpedRotation.GetTwistAngle(RotationAxisVector);
+						TargetOrientationAngleRad = FMath::UnwindRadians(TargetOrientationAngleRad + FMath::DegreesToRadians(180.0f));
+						RootMotionDeltaDirection = -RootMotionDeltaDirection;
 					}
 				}
 
+				if (bCounterCompenstateInterpolationByRootMotion)
+				{
+#if !ENABLE_ANIM_DEBUG
+					float RootMotionDeltaAngleRad;
+#endif
+					// Counter the interpolated orientation angle by the root motion direction angle delta.
+					// This prevents our interpolation from fighting the natural root motion that's flowing through the graph.
+					RootMotionDeltaAngleRad = UE::Anim::SignedAngleRadBetweenNormals(RootMotionDeltaDirection, PreviousRootMotionDeltaDirection, RotationAxisVector);
+					// Root motion may have large deltas i.e. bad blends or sudden direction changes like pivots.
+					// Clamp the maximum counter-compensation by our max allowed correction to prevent pops.
+					RootMotionDeltaAngleRad = FMath::Clamp(RootMotionDeltaAngleRad, -MaxAngleCorrectionRad, MaxAngleCorrectionRad);
+
+					ActualOrientationAngleRad = FMath::UnwindRadians(ActualOrientationAngleRad + RootMotionDeltaAngleRad);
+				}
+
 				// Rotate the root motion delta fully by the warped angle
-				const FVector WarpedRootMotionTranslationDelta = WarpedRotation.RotateVector(RootMotionDeltaTranslation);
+				const FVector WarpedRootMotionTranslationDelta = FQuat(RotationAxisVector, TargetOrientationAngleRad).RotateVector(RootMotionDeltaTranslation);
 				RootMotionTransformDelta.SetTranslation(WarpedRootMotionTranslationDelta);
 			}
 
-
-
 			// Forward the side effects of orientation warping on the root motion contribution for this sub-graph
 			const bool bRootMotionOverridden = RootMotionProvider->OverrideRootMotion(RootMotionTransformDelta, Output.CustomAttributes);
 			ensureMsgf(bRootMotionOverridden, TEXT("Graph driven Orientation Warping expected a root motion delta to be present in the attribute stream prior to warping/overriding it."));
@@ -268,19 +244,30 @@ void FAnimNode_OrientationWarping::EvaluateSkeletalControl_AnyThread(FComponentS
 	else
 	{
 		// Manual orientation warping will take the angle directly
-		ActualOrientationAngle = FRotator::NormalizeAxis(ActualOrientationAngle);
-		ActualOrientationAngle = FMath::DegreesToRadians(ActualOrientationAngle);
+		TargetOrientationAngleRad = FRotator::NormalizeAxis(OrientationAngle);
+		TargetOrientationAngleRad = FMath::DegreesToRadians(TargetOrientationAngleRad);
 	}
 
 	// Optionally interpolate the effective orientation towards the target orientation angle
 	if (RotationInterpSpeed > 0.f)
 	{
-		ActualOrientationAngle = FMath::FInterpTo(PreviousOrientationAngle, ActualOrientationAngle, DeltaSeconds, RotationInterpSpeed);
-		PreviousOrientationAngle = ActualOrientationAngle;
+		// When the orientation warping node becomes relevant, the input pose orientation may not be aligned with the desired orientation.
+		// Instead of interpolating this difference, snap to the desired orientation if it's our first update to minimize corrections over-time.
+		if (bIsFirstUpdate == false)
+		{
+			const float SmoothOrientationAngleRad = FMath::FInterpTo(ActualOrientationAngleRad, TargetOrientationAngleRad, DeltaSeconds, RotationInterpSpeed);
+			// Limit our interpolation rate to prevent pops.
+			// @TODO: Use better, more physically accurate interpolation here.
+			ActualOrientationAngleRad = FMath::Clamp(SmoothOrientationAngleRad, ActualOrientationAngleRad - MaxAngleCorrectionRad, ActualOrientationAngleRad + MaxAngleCorrectionRad);
+		}
+	}
+	else
+	{
+		ActualOrientationAngleRad = TargetOrientationAngleRad;
 	}
 
 	// Allow the alpha value of the node to affect the final rotation
-	ActualOrientationAngle *= ActualAlpha * WarpingAlpha;
+	ActualOrientationAngleRad *= ActualAlpha;
 
 #if ENABLE_ANIM_DEBUG
 	bool bDebugging = false;
@@ -298,10 +285,7 @@ void FAnimNode_OrientationWarping::EvaluateSkeletalControl_AnyThread(FComponentS
 		const FVector ActorForwardDirection = Output.AnimInstanceProxy->GetActorTransform().GetRotation().GetForwardVector();
 		FVector DebugArrowOffset = FVector::ZAxisVector * DebugDrawScale;
 
-		const bool bDrawAll = DebugIndex == 3;
-		const bool bDrawOffset = (DebugIndex == 2) || bDrawAll;
-		const bool bDrawWarping = (DebugIndex == 1) || bDrawAll;
-		if (bDrawWarping)
+		// Draw debug shapes
 		{
 			const FVector ForwardDirection = bGraphDrivenWarping
 				? ComponentTransform.GetRotation().RotateVector(LocomotionForward)
@@ -322,7 +306,7 @@ void FAnimNode_OrientationWarping::EvaluateSkeletalControl_AnyThread(FComponentS
 				ComponentTransform.GetLocation() + DebugArrowOffset + RotationDirection * 100.f * DebugDrawScale,
 				40.f * DebugDrawScale, FColor::Blue, false, 0.f, 2.f * DebugDrawScale);
 
-			const float ActualOrientationAngleDegrees = FMath::RadiansToDegrees(ActualOrientationAngle);
+			const float ActualOrientationAngleDegrees = FMath::RadiansToDegrees(ActualOrientationAngleRad);
 			const FVector WarpedRotationDirection = bGraphDrivenWarping
 				? RotationDirection.RotateAngleAxis(ActualOrientationAngleDegrees, RotationAxisVector)
 				: ActorForwardDirection.RotateAngleAxis(ActualOrientationAngleDegrees, RotationAxisVector);
@@ -332,35 +316,43 @@ void FAnimNode_OrientationWarping::EvaluateSkeletalControl_AnyThread(FComponentS
 				ComponentTransform.GetLocation() + DebugArrowOffset,
 				ComponentTransform.GetLocation() + DebugArrowOffset + WarpedRotationDirection * 100.f * DebugDrawScale,
 				40.f * DebugDrawScale, FColor::Green, false, 0.f, 2.f * DebugDrawScale);
-			DebugArrowOffset += FVector::ZAxisVector * DebugDrawScale;
 		}
-		
-		if (bDrawOffset)
-		{
-			const float HeadingOffetDegrees = FMath::RadiansToDegrees(HeadingOffset);
-			const FVector OffsetDirection = ActorForwardDirection.RotateAngleAxis(HeadingOffetDegrees, RotationAxisVector);
 
-			Output.AnimInstanceProxy->AnimDrawDebugDirectionalArrow(
-				ComponentTransform.GetLocation() + DebugArrowOffset,
-				ComponentTransform.GetLocation() + DebugArrowOffset + OffsetDirection * 100.f * DebugDrawScale,
-				40.f * DebugDrawScale, FColor::Purple, false, 0.f, 2.f * DebugDrawScale);
+		// Draw text on mesh in world space
+		{
+			TStringBuilder<1024> DebugLine;
 
-			DebugArrowOffset += FVector::ZAxisVector * DebugDrawScale;
-			Output.AnimInstanceProxy->AnimDrawDebugDirectionalArrow(
-				ComponentTransform.GetLocation() + DebugArrowOffset,
-				ComponentTransform.GetLocation() + DebugArrowOffset + ActorForwardDirection * 100.f * DebugDrawScale,
-				40.f * DebugDrawScale, FColor::Black, false, 0.f, 2.f * DebugDrawScale);
+			if (Mode == EWarpingEvaluationMode::Manual)
+			{
+				DebugLine.Appendf(TEXT("\n - Previous Orientation Angle: (%.3fd)"), FMath::RadiansToDegrees(PreviousOrientationAngleRad));
+				DebugLine.Appendf(TEXT("\n - Orientation Angle: (%.3fd)"), FMath::RadiansToDegrees(ActualOrientationAngleRad));
+				DebugLine.Appendf(TEXT("\n - Target Orientation Angle: (%.3fd)"), FMath::RadiansToDegrees(TargetOrientationAngleRad));
+			}
+			else
+			{
+				if (RotationInterpSpeed > 0.0f)
+				{
+					DebugLine.Appendf(TEXT("\n - Previous Orientation Angle: (%.3fd)"), FMath::RadiansToDegrees(PreviousOrientationAngleRad));
+					DebugLine.Appendf(TEXT("\n - rootmotionTest: (%.3fd)"), FMath::RadiansToDegrees(RootMotionDeltaAngleRad));
+				}
+				DebugLine.Appendf(TEXT("\n - Actual Orientation Angle: (%.3fd)"), FMath::RadiansToDegrees(ActualOrientationAngleRad));
+				DebugLine.Appendf(TEXT("\n - Target Orientation Angle: (%.3fd)"), FMath::RadiansToDegrees(TargetOrientationAngleRad));
+				// Locomotion angle is already in degrees.
+				DebugLine.Appendf(TEXT("\n - Locomotion Angle: (%.3fd)"), LocomotionAngle);
+				DebugLine.Appendf(TEXT("\n - Root Motion Delta: %s)"), *RootMotionTransformDelta.GetTranslation().ToString());
+				DebugLine.Appendf(TEXT("\n - Root Motion Speed: %.3fd)"), RootMotionTransformDelta.GetTranslation().Size() / DeltaSeconds);
+			}
+			Output.AnimInstanceProxy->AnimDrawDebugInWorldMessage(DebugLine.ToString(), FVector::UpVector * 50.0f, FColor::Yellow, 1.f /*TextScale*/);
 		}
 	}
 #endif
 
-	// Combine our warping and heading offsets to rotate our root bone
-	const float CombinedRootOffset = FMath::UnwindRadians(ActualOrientationAngle * DistributedBoneOrientationAlpha + HeadingOffset);
+	const float RootOffset = FMath::UnwindRadians(ActualOrientationAngleRad * DistributedBoneOrientationAlpha);
 
 	// Rotate Root Bone first, as that cheaply rotates the whole pose with one transformation.
-	if (!FMath::IsNearlyZero(CombinedRootOffset, KINDA_SMALL_NUMBER))
+	if (!FMath::IsNearlyZero(RootOffset, KINDA_SMALL_NUMBER))
 	{
-		const FQuat RootRotation = FQuat(RotationAxisVector, CombinedRootOffset);
+		const FQuat RootRotation = FQuat(RotationAxisVector, RootOffset);
 		const FCompactPoseBoneIndex RootBoneIndex(0);
 
 		FTransform RootBoneTransform(Output.Pose.GetComponentSpaceTransform(RootBoneIndex));
@@ -379,7 +371,7 @@ void FAnimNode_OrientationWarping::EvaluateSkeletalControl_AnyThread(FComponentS
 		for (int32 ArrayIndex = 0; ArrayIndex < NumSpineBones; ArrayIndex++)
 		{
 			const FOrientationWarpingSpineBoneData& BoneData = SpineBoneDataArray[ArrayIndex];
-			const FQuat SpineBoneCounterRotation = FQuat(RotationAxisVector, -ActualOrientationAngle * DistributedBoneOrientationAlpha * BoneData.Weight);
+			const FQuat SpineBoneCounterRotation = FQuat(RotationAxisVector, -ActualOrientationAngleRad * DistributedBoneOrientationAlpha * BoneData.Weight);
 			check(BoneData.Weight > 0.f);
 
 			FTransform SpineBoneTransform(Output.Pose.GetComponentSpaceTransform(BoneData.BoneIndex));
@@ -395,7 +387,7 @@ void FAnimNode_OrientationWarping::EvaluateSkeletalControl_AnyThread(FComponentS
 	// Rotate IK Foot Root
 	if (bUpdateIKFootRoot)
 	{
-		const FQuat BoneRotation = FQuat(RotationAxisVector, ActualOrientationAngle * IKFootRootOrientationAlpha);
+		const FQuat BoneRotation = FQuat(RotationAxisVector, ActualOrientationAngleRad * IKFootRootOrientationAlpha);
 
 		FTransform IKFootRootTransform(Output.Pose.GetComponentSpaceTransform(IKFootData.IKFootRootBoneIndex));
 		IKFootRootTransform.SetRotation(BoneRotation * IKFootRootTransform.GetRotation());
@@ -410,7 +402,7 @@ void FAnimNode_OrientationWarping::EvaluateSkeletalControl_AnyThread(FComponentS
 
 		if (bUpdateIKFootBones)
 		{
-			const FQuat IKFootRotation = FQuat(RotationAxisVector, -ActualOrientationAngle * IKFootRootOrientationAlpha);
+			const FQuat IKFootRotation = FQuat(RotationAxisVector, -ActualOrientationAngleRad * IKFootRootOrientationAlpha);
 
 			for (int32 ArrayIndex = 0; ArrayIndex < NumIKFootBones; ArrayIndex++)
 			{
@@ -425,6 +417,7 @@ void FAnimNode_OrientationWarping::EvaluateSkeletalControl_AnyThread(FComponentS
 	}
 
 	OutBoneTransforms.Sort(FCompareBoneTransformIndex());
+	bIsFirstUpdate = false;
 }
 
 bool FAnimNode_OrientationWarping::IsValidToEvaluate(const USkeleton* Skeleton, const FBoneContainer& RequiredBones)
@@ -552,6 +545,6 @@ void FAnimNode_OrientationWarping::InitializeBoneReferences(const FBoneContainer
 void FAnimNode_OrientationWarping::Reset(const FAnimationBaseContext& Context)
 {
 	bIsFirstUpdate = true;
-	HeadingOffset = 0.0f;
-	ComponentHeading = Context.AnimInstanceProxy->GetComponentTransform().GetRotation().GetTwistAngle(UE::Anim::GetAxisVector(RotationAxis));
+	RootMotionDeltaDirection = FVector::ZeroVector;
+	ActualOrientationAngleRad = 0.f;
 }
diff --git a/Engine/Plugins/Animation/AnimationWarping/Source/Runtime/Public/BoneControllers/AnimNode_OrientationWarping.h b/Engine/Plugins/Animation/AnimationWarping/Source/Runtime/Public/BoneControllers/AnimNode_OrientationWarping.h
--- a/Engine/Plugins/Animation/AnimationWarping/Source/Runtime/Public/BoneControllers/AnimNode_OrientationWarping.h
+++ b/Engine/Plugins/Animation/AnimationWarping/Source/Runtime/Public/BoneControllers/AnimNode_OrientationWarping.h
@@ -71,14 +71,16 @@ struct ANIMATIONWARPINGRUNTIME_API FAnimNode_OrientationWarping : public FAnimNo
 	UPROPERTY(EditAnywhere, Category=Settings, meta=(ClampMin="0.0"))
 	float RotationInterpSpeed = 10.f;
 
-	UPROPERTY(EditAnywhere, Category = Experimental, meta=(PinHiddenByDefault, ClampMin="0.0", ClampMax="1.0"))
-	float WarpingAlpha = 1.0f;
+	// Max correction we're allowed to do per-second when using interpolation.
+	// This minimizes pops when we have a large difference between current and target orientation.
+	UPROPERTY(EditAnywhere, Category=Settings, meta=(ClampMin="0.0", EditCondition="RotationInterpSpeed > 0.0f"))
+	float MaxCorrectionRateDegrees = 720.f;
 
-	UPROPERTY(EditAnywhere, Category = Experimental, meta=(PinHiddenByDefault, ClampMin="0.0", ClampMax="1.0"))
-	float OffsetAlpha = 0.0f;
-
-	UPROPERTY(EditAnywhere, Category = Experimental, meta=(PinHiddenByDefault, ClampMin="0.0", ClampMax="180.0"))
-	float MaxOffsetAngle = 70.0f;
+	// Whether to counter compensate interpolation by the animated root motion angle change over time.
+	// This helps to conserve the motion from our animation.
+	// Disable this if your root motion is expected to be jittery, and you want orientation warping to smooth it out.
+	UPROPERTY(EditAnywhere, Category=Settings, meta=(EditCondition="RotationInterpSpeed > 0.0f"))
+	bool bCounterCompenstateInterpolationByRootMotion = true;
 
 #if WITH_EDITORONLY_DATA
 	// Scale all debug drawing visualization by a factor
@@ -151,22 +153,12 @@ private:
 
 	// Computed IK bone indices for the specified foot definitions 
 	FOrientationWarpingFootData IKFootData;
-
-	// Internal previous frame root motion delta direction
-	FVector PreviousRootMotionDeltaDirection = FVector::ZeroVector;
-
-	// Internal previous frame orientation warping angle
-	float PreviousOrientationAngle = 0.f;
 	
-	// Internal orientation warping angle
-	float ActualOrientationAngle = 0.f;
-
-	// Our component's heading
-	float ComponentHeading = 0.0;
+	// Internal current frame root motion delta direction
+	FVector RootMotionDeltaDirection = FVector::ZeroVector;
 
-	// Our accumulated heading offset based on component frame deltas
-	float HeadingOffset = 0.0;
-	float LastOffsetAlpha = 0.0f;
+	// Internal orientation warping angle
+	float ActualOrientationAngleRad = 0.f;
 
 	FGraphTraversalCounter UpdateCounter;
 	bool bIsFirstUpdate = false;
