Changed orientation warping debug text
#rnx

diff --git a/Engine/Plugins/Animation/AnimationWarping/Source/Runtime/Private/BoneControllers/AnimNode_OrientationWarping.cpp b/Engine/Plugins/Animation/AnimationWarping/Source/Runtime/Private/BoneControllers/AnimNode_OrientationWarping.cpp
--- a/Engine/Plugins/Animation/AnimationWarping/Source/Runtime/Private/BoneControllers/AnimNode_OrientationWarping.cpp
+++ b/Engine/Plugins/Animation/AnimationWarping/Source/Runtime/Private/BoneControllers/AnimNode_OrientationWarping.cpp
@@ -333,7 +333,7 @@ void FAnimNode_OrientationWarping::EvaluateSkeletalControl_AnyThread(FComponentS
 				if (RotationInterpSpeed > 0.0f)
 				{
 					DebugLine.Appendf(TEXT("\n - Previous Orientation Angle: (%.3fd)"), FMath::RadiansToDegrees(PreviousOrientationAngleRad));
-					DebugLine.Appendf(TEXT("\n - rootmotionTest: (%.3fd)"), FMath::RadiansToDegrees(RootMotionDeltaAngleRad));
+					DebugLine.Appendf(TEXT("\n - Root Motion Frame Delta Angle: (%.3fd)"), FMath::RadiansToDegrees(RootMotionDeltaAngleRad));
 				}
 				DebugLine.Appendf(TEXT("\n - Actual Orientation Angle: (%.3fd)"), FMath::RadiansToDegrees(ActualOrientationAngleRad));
 				DebugLine.Appendf(TEXT("\n - Target Orientation Angle: (%.3fd)"), FMath::RadiansToDegrees(TargetOrientationAngleRad));
