Fix orientation warping on first frame not snapping to desired direction
#rnx

diff --git a/Engine/Plugins/Animation/AnimationWarping/Source/Runtime/Private/BoneControllers/AnimNode_OrientationWarping.cpp b/Engine/Plugins/Animation/AnimationWarping/Source/Runtime/Private/BoneControllers/AnimNode_OrientationWarping.cpp
--- a/Engine/Plugins/Animation/AnimationWarping/Source/Runtime/Private/BoneControllers/AnimNode_OrientationWarping.cpp
+++ b/Engine/Plugins/Animation/AnimationWarping/Source/Runtime/Private/BoneControllers/AnimNode_OrientationWarping.cpp
@@ -249,17 +249,14 @@ void FAnimNode_OrientationWarping::EvaluateSkeletalControl_AnyThread(FComponentS
 	}
 
 	// Optionally interpolate the effective orientation towards the target orientation angle
-	if (RotationInterpSpeed > 0.f)
+	// When the orientation warping node becomes relevant, the input pose orientation may not be aligned with the desired orientation.
+	// Instead of interpolating this difference, snap to the desired orientation if it's our first update to minimize corrections over-time.
+	if ((RotationInterpSpeed > 0.f) && !bIsFirstUpdate)
 	{
-		// When the orientation warping node becomes relevant, the input pose orientation may not be aligned with the desired orientation.
-		// Instead of interpolating this difference, snap to the desired orientation if it's our first update to minimize corrections over-time.
-		if (bIsFirstUpdate == false)
-		{
-			const float SmoothOrientationAngleRad = FMath::FInterpTo(ActualOrientationAngleRad, TargetOrientationAngleRad, DeltaSeconds, RotationInterpSpeed);
-			// Limit our interpolation rate to prevent pops.
-			// @TODO: Use better, more physically accurate interpolation here.
-			ActualOrientationAngleRad = FMath::Clamp(SmoothOrientationAngleRad, ActualOrientationAngleRad - MaxAngleCorrectionRad, ActualOrientationAngleRad + MaxAngleCorrectionRad);
-		}
+		const float SmoothOrientationAngleRad = FMath::FInterpTo(ActualOrientationAngleRad, TargetOrientationAngleRad, DeltaSeconds, RotationInterpSpeed);
+		// Limit our interpolation rate to prevent pops.
+		// @TODO: Use better, more physically accurate interpolation here.
+		ActualOrientationAngleRad = FMath::Clamp(SmoothOrientationAngleRad, ActualOrientationAngleRad - MaxAngleCorrectionRad, ActualOrientationAngleRad + MaxAngleCorrectionRad);
 	}
 	else
 	{
