Pose Warping: Add on screen debug viz for stride warping

diff --git a/Engine/Plugins/Animation/AnimationWarping/Source/Runtime/Private/BoneControllers/AnimNode_StrideWarping.cpp b/Engine/Plugins/Animation/AnimationWarping/Source/Runtime/Private/BoneControllers/AnimNode_StrideWarping.cpp
--- a/Engine/Plugins/Animation/AnimationWarping/Source/Runtime/Private/BoneControllers/AnimNode_StrideWarping.cpp
+++ b/Engine/Plugins/Animation/AnimationWarping/Source/Runtime/Private/BoneControllers/AnimNode_StrideWarping.cpp
@@ -33,7 +33,7 @@ void FAnimNode_StrideWarping::GatherDebugData(FNodeDebugData& DebugData)
 			DebugLine += FString::Printf(TEXT("\n - Root Motion Direction: (%s)"), *(CachedRootMotionDeltaTranslation.GetSafeNormal().ToCompactString()));
 			DebugLine += FString::Printf(TEXT("\n - Root Motion Speed: (%.3fd)"), CachedRootMotionDeltaSpeed);
 #endif
-			DebugLine += FString::Printf(TEXT("\n - Min Locomotion Speed Threshold: (%.3fd)"), MinRootMotionSpeedThreshold);
+			DebugLine += FString::Printf(TEXT("\n - Min Root Motion Speed Threshold: (%.3fd)"), MinRootMotionSpeedThreshold);
 		}
 		DebugLine += FString::Printf(TEXT("\n - Floor Normal: (%s)"), *(FloorNormalDirection.Value.ToCompactString()));
 		DebugLine += FString::Printf(TEXT("\n - Gravity Direction: (%s)"), *(GravityDirection.Value.ToCompactString()));
@@ -268,6 +268,34 @@ void FAnimNode_StrideWarping::EvaluateSkeletalControl_AnyThread(FComponentSpaceP
 			AnimInstanceProxy->GetComponentTransform().TransformPosition(InitialPelvisLocation),
 			AnimInstanceProxy->GetComponentTransform().TransformPosition(InitialPelvisLocation + ActualStrideDirection * ActualStrideScale * 100.f * DebugDrawScale),
 			40.f * DebugDrawScale, FColor::Red, false, 0.f, 2.f * DebugDrawScale);
+
+		// Draw text on mesh in world space
+		{
+			TStringBuilder<1024> DebugLine;
+
+			DebugLine.Appendf(TEXT("\n - Stride Scale: (%.3fd)"), ActualStrideScale);
+			DebugLine.Appendf(TEXT("\n - Stride Direction: (%s)"), *(ActualStrideDirection.ToCompactString()));
+			
+			if (Mode == EWarpingEvaluationMode::Graph)
+			{
+				DebugLine.Appendf(TEXT("\n - Locomotion Speed: (%.3fd)"), LocomotionSpeed);
+
+#if WITH_EDITORONLY_DATA
+				if (bFoundRootMotionAttribute)
+				{
+					DebugLine.Appendf(TEXT("\n - Root Motion: Direction(%s), Speed(%.3fd)"), *(CachedRootMotionDeltaTranslation.GetSafeNormal().ToCompactString()), CachedRootMotionDeltaSpeed);
+				}
+				else
+				{
+					DebugLine.Appendf(TEXT("\n - Root Motion: NO ATTRIBUTE FOUND"), *(ActualStrideDirection.ToCompactString()));
+				}
+#endif
+
+				DebugLine.Appendf(TEXT("\n - Min Root Motion Speed Threshold: (%.3fd)"), MinRootMotionSpeedThreshold);
+			}
+
+			AnimInstanceProxy->AnimDrawDebugInWorldMessage(DebugLine.ToString(), FVector::UpVector * 50.0f, FColor::Yellow, 1.f /*TextScale*/);
+		}
 	}
 #endif
 	// Add adjusted pelvis transform
