Experimental: Add manual root motion mode for orientation warping (used for prediction)

diff --git a/Engine/Plugins/Animation/AnimationWarping/Source/Runtime/Private/BoneControllers/AnimNode_OrientationWarping.cpp b/Engine/Plugins/Animation/AnimationWarping/Source/Runtime/Private/BoneControllers/AnimNode_OrientationWarping.cpp
--- a/Engine/Plugins/Animation/AnimationWarping/Source/Runtime/Private/BoneControllers/AnimNode_OrientationWarping.cpp
+++ b/Engine/Plugins/Animation/AnimationWarping/Source/Runtime/Private/BoneControllers/AnimNode_OrientationWarping.cpp
@@ -5,6 +5,7 @@
 #include "Animation/AnimNodeFunctionRef.h"
 #include "Animation/AnimRootMotionProvider.h"
 #include "HAL/IConsoleManager.h"
+#include "Animation/AnimTrace.h"
 
 #include UE_INLINE_GENERATED_CPP_BY_NAME(AnimNode_OrientationWarping)
 
@@ -183,7 +184,8 @@ void FAnimNode_OrientationWarping::EvaluateSkeletalControl_AnyThread(FComponentS
 			const FQuat SkeletalMeshRelativeRotation = SkeletalMeshRelativeTransform.GetRotation();
 			LocomotionForward = SkeletalMeshRelativeRotation.UnrotateVector(LocomotionRotation.GetForwardVector()).GetSafeNormal();
 
-			const FVector RootMotionDeltaTranslation = RootMotionTransformDelta.GetTranslation();
+			// @todo: Graph mode using a "manual value" makes no sense. Restructure logic to address this in the future.
+			const FVector RootMotionDeltaTranslation = bUseManualRootMotionVelocity ? (ManualRootMotionVelocity * DeltaSeconds) : RootMotionTransformDelta.GetTranslation();
 
 			const float RootMotionDeltaSpeed = RootMotionDeltaTranslation.Size() / DeltaSeconds;
 			if (RootMotionDeltaSpeed < MinRootMotionSpeedThreshold)
@@ -221,10 +223,14 @@ void FAnimNode_OrientationWarping::EvaluateSkeletalControl_AnyThread(FComponentS
 					// This prevents our interpolation from fighting the natural root motion that's flowing through the graph.
 					RootMotionDeltaAngleRad = UE::Anim::SignedAngleRadBetweenNormals(RootMotionDeltaDirection, PreviousRootMotionDeltaDirection, RotationAxisVector);
 					// Root motion may have large deltas i.e. bad blends or sudden direction changes like pivots.
-					// Clamp the maximum counter-compensation by our max allowed correction to prevent pops.
-					RootMotionDeltaAngleRad = FMath::Clamp(RootMotionDeltaAngleRad, -MaxAngleCorrectionRad, MaxAngleCorrectionRad);
-
-					ActualOrientationAngleRad = FMath::UnwindRadians(ActualOrientationAngleRad + RootMotionDeltaAngleRad);
+					// If there's an instantaneous pop in root motion direction, this is likely a pivot.
+					const float MaxRootMotionDeltaToCompensateRad = FMath::DegreesToRadians(MaxRootMotionDeltaToCompensateDegrees);
+					if (FMath::Abs(RootMotionDeltaAngleRad) < MaxRootMotionDeltaToCompensateRad)
+					{
+						// Otherwise, clamp the maximum counter-compensation by our max allowed correction to prevent pops.
+						RootMotionDeltaAngleRad = FMath::Clamp(RootMotionDeltaAngleRad, -MaxAngleCorrectionRad, MaxAngleCorrectionRad);
+						ActualOrientationAngleRad = FMath::UnwindRadians(ActualOrientationAngleRad + RootMotionDeltaAngleRad);
+					}
 				}
 
 				// Rotate the root motion delta fully by the warped angle
@@ -320,11 +326,14 @@ void FAnimNode_OrientationWarping::EvaluateSkeletalControl_AnyThread(FComponentS
 		{
 			TStringBuilder<1024> DebugLine;
 
+			const float PreviousOrientationAngleDegrees = FMath::RadiansToDegrees(PreviousOrientationAngleRad);
+			const float ActualOrientationAngleDegrees = FMath::RadiansToDegrees(ActualOrientationAngleRad);
+			const float TargetOrientationAngleDegrees = FMath::RadiansToDegrees(TargetOrientationAngleRad);
 			if (Mode == EWarpingEvaluationMode::Manual)
 			{
-				DebugLine.Appendf(TEXT("\n - Previous Orientation Angle: (%.3fd)"), FMath::RadiansToDegrees(PreviousOrientationAngleRad));
-				DebugLine.Appendf(TEXT("\n - Orientation Angle: (%.3fd)"), FMath::RadiansToDegrees(ActualOrientationAngleRad));
-				DebugLine.Appendf(TEXT("\n - Target Orientation Angle: (%.3fd)"), FMath::RadiansToDegrees(TargetOrientationAngleRad));
+				DebugLine.Appendf(TEXT("\n - Previous Orientation Angle: (%.3fd)"), PreviousOrientationAngleDegrees);
+				DebugLine.Appendf(TEXT("\n - Orientation Angle: (%.3fd)"), ActualOrientationAngleDegrees);
+				DebugLine.Appendf(TEXT("\n - Target Orientation Angle: (%.3fd)"), TargetOrientationAngleRad);
 			}
 			else
 			{
@@ -345,6 +354,29 @@ void FAnimNode_OrientationWarping::EvaluateSkeletalControl_AnyThread(FComponentS
 	}
 #endif
 
+#if ANIM_TRACE_ENABLED
+	{
+		const float PreviousOrientationAngleDegrees = FMath::RadiansToDegrees(PreviousOrientationAngleRad);
+		const float ActualOrientationAngleDegrees = FMath::RadiansToDegrees(ActualOrientationAngleRad);
+		const float TargetOrientationAngleDegrees = FMath::RadiansToDegrees(TargetOrientationAngleRad);
+
+		TRACE_ANIM_NODE_VALUE(Output, TEXT("Previous OrientationAngle Degrees"), PreviousOrientationAngleDegrees);
+		TRACE_ANIM_NODE_VALUE(Output, TEXT("Actual Orientation Angle Degrees"), ActualOrientationAngleDegrees);
+		TRACE_ANIM_NODE_VALUE(Output, TEXT("Target Orientation Angle Degrees"), TargetOrientationAngleDegrees);
+
+		if (Mode == EWarpingEvaluationMode::Graph)
+		{
+			const float RootMotionDeltaAngleDegrees = FMath::RadiansToDegrees(RootMotionDeltaAngleRad);
+			TRACE_ANIM_NODE_VALUE(Output, TEXT("Root Motion Delta Angle Degrees"), RootMotionDeltaAngleDegrees);
+			TRACE_ANIM_NODE_VALUE(Output, TEXT("Locomotion Angle"), LocomotionAngle);
+			TRACE_ANIM_NODE_VALUE(Output, TEXT("Root Motion Translation Delta"), RootMotionTransformDelta.GetTranslation());
+
+			const float RootMotionSpeed = RootMotionTransformDelta.GetTranslation().Size() / DeltaSeconds;
+			TRACE_ANIM_NODE_VALUE(Output, TEXT("Root Motion Speed"), RootMotionSpeed);
+		}
+	}
+#endif
+
 	const float RootOffset = FMath::UnwindRadians(ActualOrientationAngleRad * DistributedBoneOrientationAlpha);
 
 	// Rotate Root Bone first, as that cheaply rotates the whole pose with one transformation.
diff --git a/Engine/Plugins/Animation/AnimationWarping/Source/Runtime/Public/BoneControllers/AnimNode_OrientationWarping.h b/Engine/Plugins/Animation/AnimationWarping/Source/Runtime/Public/BoneControllers/AnimNode_OrientationWarping.h
--- a/Engine/Plugins/Animation/AnimationWarping/Source/Runtime/Public/BoneControllers/AnimNode_OrientationWarping.h
+++ b/Engine/Plugins/Animation/AnimationWarping/Source/Runtime/Public/BoneControllers/AnimNode_OrientationWarping.h
@@ -76,12 +76,22 @@ struct ANIMATIONWARPINGRUNTIME_API FAnimNode_OrientationWarping : public FAnimNo
 	UPROPERTY(EditAnywhere, Category=Settings, meta=(ClampMin="0.0", EditCondition="RotationInterpSpeed > 0.0f"))
 	float MaxCorrectionRateDegrees = 720.f;
 
+	// Don't compensate our interpolator when the instantaneous root motion delta is higher than this. This is likely a pivot.
+	UPROPERTY(EditAnywhere, Category=Settings, meta=(ClampMin="0.0", EditCondition="RotationInterpSpeed > 0.0f"))
+	float MaxRootMotionDeltaToCompensateDegrees = 45.f;
+
 	// Whether to counter compensate interpolation by the animated root motion angle change over time.
 	// This helps to conserve the motion from our animation.
 	// Disable this if your root motion is expected to be jittery, and you want orientation warping to smooth it out.
 	UPROPERTY(EditAnywhere, Category=Settings, meta=(EditCondition="RotationInterpSpeed > 0.0f"))
 	bool bCounterCompenstateInterpolationByRootMotion = true;
 
+	UPROPERTY(EditAnywhere, Category=Experimental, meta=(EditCondition="Mode == EWarpingEvaluationMode::Graph"))
+	bool bUseManualRootMotionVelocity = false;
+
+	UPROPERTY(EditAnywhere, BlueprintReadWrite, Category=Experimental, meta=(PinHiddenByDefault, EditCondition="bUseManualRootMotionVelocity && (Mode == EWarpingEvaluationMode::Graph)"))
+	FVector ManualRootMotionVelocity;
+
 #if WITH_EDITORONLY_DATA
 	// Scale all debug drawing visualization by a factor
 	UPROPERTY(EditAnywhere, Category=Debug, meta=(ClampMin="0.0"))
