Orientation Warping: Initialize UProperty properly

diff --git a/Engine/Plugins/Animation/AnimationWarping/Source/Runtime/Public/BoneControllers/AnimNode_OrientationWarping.h b/Engine/Plugins/Animation/AnimationWarping/Source/Runtime/Public/BoneControllers/AnimNode_OrientationWarping.h
--- a/Engine/Plugins/Animation/AnimationWarping/Source/Runtime/Public/BoneControllers/AnimNode_OrientationWarping.h
+++ b/Engine/Plugins/Animation/AnimationWarping/Source/Runtime/Public/BoneControllers/AnimNode_OrientationWarping.h
@@ -90,7 +90,7 @@ struct ANIMATIONWARPINGRUNTIME_API FAnimNode_OrientationWarping : public FAnimNo
 	bool bUseManualRootMotionVelocity = false;
 
 	UPROPERTY(EditAnywhere, BlueprintReadWrite, Category=Experimental, meta=(PinHiddenByDefault, EditCondition="bUseManualRootMotionVelocity && (Mode == EWarpingEvaluationMode::Graph)"))
-	FVector ManualRootMotionVelocity;
+	FVector ManualRootMotionVelocity = FVector::ZeroVector;
 
 #if WITH_EDITORONLY_DATA
 	// Scale all debug drawing visualization by a factor
