Orientation Warping: Reset manual root motion on re-init

diff --git a/Engine/Plugins/Animation/AnimationWarping/Source/Runtime/Private/BoneControllers/AnimNode_OrientationWarping.cpp b/Engine/Plugins/Animation/AnimationWarping/Source/Runtime/Private/BoneControllers/AnimNode_OrientationWarping.cpp
--- a/Engine/Plugins/Animation/AnimationWarping/Source/Runtime/Private/BoneControllers/AnimNode_OrientationWarping.cpp
+++ b/Engine/Plugins/Animation/AnimationWarping/Source/Runtime/Private/BoneControllers/AnimNode_OrientationWarping.cpp
@@ -576,5 +576,6 @@ void FAnimNode_OrientationWarping::Reset(const FAnimationBaseContext& Context)
 {
 	bIsFirstUpdate = true;
 	RootMotionDeltaDirection = FVector::ZeroVector;
+	ManualRootMotionVelocity = FVector::ZeroVector;
 	ActualOrientationAngleRad = 0.f;
 }
