Orientation warping: Make "Use manual root motion velocity pinnable.
Debug draw the correct root motion delta

diff --git a/Engine/Plugins/Animation/AnimationWarping/Source/Runtime/Private/BoneControllers/AnimNode_OrientationWarping.cpp b/Engine/Plugins/Animation/AnimationWarping/Source/Runtime/Private/BoneControllers/AnimNode_OrientationWarping.cpp
--- a/Engine/Plugins/Animation/AnimationWarping/Source/Runtime/Private/BoneControllers/AnimNode_OrientationWarping.cpp
+++ b/Engine/Plugins/Animation/AnimationWarping/Source/Runtime/Private/BoneControllers/AnimNode_OrientationWarping.cpp
@@ -185,7 +185,12 @@ void FAnimNode_OrientationWarping::EvaluateSkeletalControl_AnyThread(FComponentS
 			LocomotionForward = SkeletalMeshRelativeRotation.UnrotateVector(LocomotionRotation.GetForwardVector()).GetSafeNormal();
 
 			// @todo: Graph mode using a "manual value" makes no sense. Restructure logic to address this in the future.
-			const FVector RootMotionDeltaTranslation = bUseManualRootMotionVelocity ? (ManualRootMotionVelocity * DeltaSeconds) : RootMotionTransformDelta.GetTranslation();
+			if (bUseManualRootMotionVelocity)
+			{
+				RootMotionTransformDelta.SetTranslation(ManualRootMotionVelocity * DeltaSeconds);
+			}
+
+			const FVector RootMotionDeltaTranslation = RootMotionTransformDelta.GetTranslation();
 
 			const float RootMotionDeltaSpeed = RootMotionDeltaTranslation.Size() / DeltaSeconds;
 			if (RootMotionDeltaSpeed < MinRootMotionSpeedThreshold)
diff --git a/Engine/Plugins/Animation/AnimationWarping/Source/Runtime/Public/BoneControllers/AnimNode_OrientationWarping.h b/Engine/Plugins/Animation/AnimationWarping/Source/Runtime/Public/BoneControllers/AnimNode_OrientationWarping.h
--- a/Engine/Plugins/Animation/AnimationWarping/Source/Runtime/Public/BoneControllers/AnimNode_OrientationWarping.h
+++ b/Engine/Plugins/Animation/AnimationWarping/Source/Runtime/Public/BoneControllers/AnimNode_OrientationWarping.h
@@ -86,10 +86,10 @@ struct ANIMATIONWARPINGRUNTIME_API FAnimNode_OrientationWarping : public FAnimNo
 	UPROPERTY(EditAnywhere, Category=Settings, meta=(EditCondition="RotationInterpSpeed > 0.0f"))
 	bool bCounterCompenstateInterpolationByRootMotion = true;
 
-	UPROPERTY(EditAnywhere, Category=Experimental, meta=(EditCondition="Mode == EWarpingEvaluationMode::Graph"))
+	UPROPERTY(EditAnywhere, Category=Experimental, meta=(PinHiddenByDefault))
 	bool bUseManualRootMotionVelocity = false;
 
-	UPROPERTY(EditAnywhere, BlueprintReadWrite, Category=Experimental, meta=(PinHiddenByDefault, EditCondition="bUseManualRootMotionVelocity && (Mode == EWarpingEvaluationMode::Graph)"))
+	UPROPERTY(EditAnywhere, BlueprintReadWrite, Category=Experimental, meta=(PinHiddenByDefault))
 	FVector ManualRootMotionVelocity = FVector::ZeroVector;
 
 #if WITH_EDITORONLY_DATA
