Orientation Warping: Multiply by global blend weight to reduce contribution on samples that are not at full weight

diff --git a/Engine/Plugins/Animation/AnimationWarping/Source/Runtime/Private/BoneControllers/AnimNode_OrientationWarping.cpp b/Engine/Plugins/Animation/AnimationWarping/Source/Runtime/Private/BoneControllers/AnimNode_OrientationWarping.cpp
--- a/Engine/Plugins/Animation/AnimationWarping/Source/Runtime/Private/BoneControllers/AnimNode_OrientationWarping.cpp
+++ b/Engine/Plugins/Animation/AnimationWarping/Source/Runtime/Private/BoneControllers/AnimNode_OrientationWarping.cpp
@@ -103,6 +103,7 @@ void FAnimNode_OrientationWarping::UpdateInternal(const FAnimationUpdateContext&
 		Reset(Context);
 	}
 	UpdateCounter.SynchronizeWith(Context.AnimInstanceProxy->GetUpdateCounter());
+	BlendWeight = Context.GetFinalBlendWeight();
 }
 
 void FAnimNode_OrientationWarping::EvaluateSkeletalControl_AnyThread(FComponentSpacePoseContext& Output, TArray<FBoneTransform>& OutBoneTransforms)
@@ -278,6 +279,11 @@ void FAnimNode_OrientationWarping::EvaluateSkeletalControl_AnyThread(FComponentS
 	// Allow the alpha value of the node to affect the final rotation
 	ActualOrientationAngleRad *= ActualAlpha;
 
+	if (bScaleByGlobalBlendWeight)
+	{
+		ActualOrientationAngleRad *= BlendWeight;
+	}
+
 #if ENABLE_ANIM_DEBUG
 	bool bDebugging = false;
 #if WITH_EDITORONLY_DATA
diff --git a/Engine/Plugins/Animation/AnimationWarping/Source/Runtime/Public/BoneControllers/AnimNode_OrientationWarping.h b/Engine/Plugins/Animation/AnimationWarping/Source/Runtime/Public/BoneControllers/AnimNode_OrientationWarping.h
--- a/Engine/Plugins/Animation/AnimationWarping/Source/Runtime/Public/BoneControllers/AnimNode_OrientationWarping.h
+++ b/Engine/Plugins/Animation/AnimationWarping/Source/Runtime/Public/BoneControllers/AnimNode_OrientationWarping.h
@@ -86,6 +86,9 @@ struct ANIMATIONWARPINGRUNTIME_API FAnimNode_OrientationWarping : public FAnimNo
 	UPROPERTY(EditAnywhere, Category=Settings, meta=(EditCondition="RotationInterpSpeed > 0.0f"))
 	bool bCounterCompenstateInterpolationByRootMotion = true;
 
+	UPROPERTY(EditAnywhere, Category=Experimental, meta=(PinHiddenByDefault))
+	bool bScaleByGlobalBlendWeight = false;
+
 	UPROPERTY(EditAnywhere, Category=Experimental, meta=(PinHiddenByDefault))
 	bool bUseManualRootMotionVelocity = false;
 
@@ -169,6 +172,7 @@ private:
 
 	// Internal orientation warping angle
 	float ActualOrientationAngleRad = 0.f;
+	float BlendWeight = 0.0f;
 
 	FGraphTraversalCounter UpdateCounter;
 	bool bIsFirstUpdate = false;
