Made stride warping scale pinnable to allow gameplay state to modify settings.

#rnx

#changelist validated

#virtualized

diff --git a/Engine/Plugins/Animation/AnimationWarping/Source/Runtime/Public/BoneControllers/AnimNode_StrideWarping.h b/Engine/Plugins/Animation/AnimationWarping/Source/Runtime/Public/BoneControllers/AnimNode_StrideWarping.h
--- a/Engine/Plugins/Animation/AnimationWarping/Source/Runtime/Public/BoneControllers/AnimNode_StrideWarping.h
+++ b/Engine/Plugins/Animation/AnimationWarping/Source/Runtime/Public/BoneControllers/AnimNode_StrideWarping.h
@@ -77,7 +77,7 @@ struct ANIMATIONWARPINGRUNTIME_API FAnimNode_StrideWarping : public FAnimNode_Sk
 	TArray<FStrideWarpingFootDefinition> FootDefinitions;
 
 	// Modifies the final stride scale value by optionally clamping and/or interpolating
-	UPROPERTY(EditAnywhere, Category=Settings)
+	UPROPERTY(EditAnywhere, Category=Settings, meta=(PinHiddenByDefault))
 	FInputClampConstants StrideScaleModifier;
 
 	// Floor normal direction, this value will internally convert into a corresponding Component-space representation prior to warping
