Orientation Warping: Remove MaxCorrectionRate (introduces jitter), and changed to "max correction angle"

diff --git a/Engine/Plugins/Animation/AnimationWarping/Source/Runtime/Private/BoneControllers/AnimNode_OrientationWarping.cpp b/Engine/Plugins/Animation/AnimationWarping/Source/Runtime/Private/BoneControllers/AnimNode_OrientationWarping.cpp
--- a/Engine/Plugins/Animation/AnimationWarping/Source/Runtime/Private/BoneControllers/AnimNode_OrientationWarping.cpp
+++ b/Engine/Plugins/Animation/AnimationWarping/Source/Runtime/Private/BoneControllers/AnimNode_OrientationWarping.cpp
@@ -114,7 +114,7 @@ void FAnimNode_OrientationWarping::EvaluateSkeletalControl_AnyThread(FComponentS
 	float TargetOrientationAngleRad;
 
 	const float DeltaSeconds = Output.AnimInstanceProxy->GetDeltaSeconds();
-	const float MaxAngleCorrectionRad = FMath::DegreesToRadians(MaxCorrectionRateDegrees) * DeltaSeconds;
+	const float MaxAngleCorrectionRad = FMath::DegreesToRadians(MaxCorrectionDegrees);
 	const FVector RotationAxisVector = UE::Anim::GetAxisVector(RotationAxis);
 	FVector LocomotionForward = FVector::ZeroVector;
 
@@ -233,8 +233,6 @@ void FAnimNode_OrientationWarping::EvaluateSkeletalControl_AnyThread(FComponentS
 					const float MaxRootMotionDeltaToCompensateRad = FMath::DegreesToRadians(MaxRootMotionDeltaToCompensateDegrees);
 					if (FMath::Abs(RootMotionDeltaAngleRad) < MaxRootMotionDeltaToCompensateRad)
 					{
-						// Otherwise, clamp the maximum counter-compensation by our max allowed correction to prevent pops.
-						RootMotionDeltaAngleRad = FMath::Clamp(RootMotionDeltaAngleRad, -MaxAngleCorrectionRad, MaxAngleCorrectionRad);
 						ActualOrientationAngleRad = FMath::UnwindRadians(ActualOrientationAngleRad + RootMotionDeltaAngleRad);
 					}
 				}
@@ -276,6 +274,7 @@ void FAnimNode_OrientationWarping::EvaluateSkeletalControl_AnyThread(FComponentS
 		ActualOrientationAngleRad = TargetOrientationAngleRad;
 	}
 
+	ActualOrientationAngleRad = FMath::Clamp(ActualOrientationAngleRad, -MaxAngleCorrectionRad, MaxAngleCorrectionRad);
 	// Allow the alpha value of the node to affect the final rotation
 	ActualOrientationAngleRad *= ActualAlpha;
 
diff --git a/Engine/Plugins/Animation/AnimationWarping/Source/Runtime/Public/BoneControllers/AnimNode_OrientationWarping.h b/Engine/Plugins/Animation/AnimationWarping/Source/Runtime/Public/BoneControllers/AnimNode_OrientationWarping.h
--- a/Engine/Plugins/Animation/AnimationWarping/Source/Runtime/Public/BoneControllers/AnimNode_OrientationWarping.h
+++ b/Engine/Plugins/Animation/AnimationWarping/Source/Runtime/Public/BoneControllers/AnimNode_OrientationWarping.h
@@ -74,7 +74,7 @@ struct ANIMATIONWARPINGRUNTIME_API FAnimNode_OrientationWarping : public FAnimNo
 	// Max correction we're allowed to do per-second when using interpolation.
 	// This minimizes pops when we have a large difference between current and target orientation.
 	UPROPERTY(EditAnywhere, Category=Settings, meta=(ClampMin="0.0", EditCondition="RotationInterpSpeed > 0.0f"))
-	float MaxCorrectionRateDegrees = 720.f;
+	float MaxCorrectionDegrees = 180.f;
 
 	// Don't compensate our interpolator when the instantaneous root motion delta is higher than this. This is likely a pivot.
 	UPROPERTY(EditAnywhere, Category=Settings, meta=(ClampMin="0.0", EditCondition="RotationInterpSpeed > 0.0f"))
