Expose orientation warping distributed bone alpha to blueprint

diff --git a/Engine/Plugins/Animation/AnimationWarping/Source/Runtime/Public/BoneControllers/AnimNode_OrientationWarping.h b/Engine/Plugins/Animation/AnimationWarping/Source/Runtime/Public/BoneControllers/AnimNode_OrientationWarping.h
--- a/Engine/Plugins/Animation/AnimationWarping/Source/Runtime/Public/BoneControllers/AnimNode_OrientationWarping.h
+++ b/Engine/Plugins/Animation/AnimationWarping/Source/Runtime/Public/BoneControllers/AnimNode_OrientationWarping.h
@@ -63,7 +63,7 @@ struct ANIMATIONWARPINGRUNTIME_API FAnimNode_OrientationWarping : public FAnimNo
 	TEnumAsByte<EAxis::Type> RotationAxis = EAxis::Z;
 
 	// Specifies how much rotation is applied to the character body versus IK feet
-	UPROPERTY(EditAnywhere, Category=Settings, meta=(ClampMin="0.0", ClampMax="1.0"))
+	UPROPERTY(EditAnywhere, Category=Settings, meta=(ClampMin="0.0", ClampMax="1.0", PinHiddenByDefault))
 	float DistributedBoneOrientationAlpha = 0.5f;
 
 	// Specifies the interpolation speed (in Alpha per second) towards reaching the final warped rotation angle
