Add Locomotion Direction option to Orientation Warping, and add Rewind Debugger debug rendering for both Orientation Warping and Offset Root Bone

diff --git a/Engine/Plugins/Animation/AnimationWarping/Source/Editor/Private/AnimGraph/AnimGraphNode_OrientationWarping.cpp b/Engine/Plugins/Animation/AnimationWarping/Source/Editor/Private/AnimGraph/AnimGraphNode_OrientationWarping.cpp
--- a/Engine/Plugins/Animation/AnimationWarping/Source/Editor/Private/AnimGraph/AnimGraphNode_OrientationWarping.cpp
+++ b/Engine/Plugins/Animation/AnimationWarping/Source/Editor/Private/AnimGraph/AnimGraphNode_OrientationWarping.cpp
@@ -48,6 +48,10 @@ void UAnimGraphNode_OrientationWarping::CustomizePinData(UEdGraphPin* Pin, FName
 	{
 		Pin->bHidden = (Node.Mode == EWarpingEvaluationMode::Manual);
 	}
+	if (Pin->PinName == GET_MEMBER_NAME_STRING_CHECKED(FAnimNode_OrientationWarping, LocomotionDirection))
+   	{
+   		Pin->bHidden = (Node.Mode == EWarpingEvaluationMode::Manual);
+   	}
 
 	if (Pin->PinName == GET_MEMBER_NAME_STRING_CHECKED(FAnimNode_OrientationWarping, LocomotionAngleDeltaThreshold))
 	{
@@ -106,6 +110,7 @@ void UAnimGraphNode_OrientationWarping::CustomizeDetails(IDetailLayoutBuilder& D
 	if (Node.Mode == EWarpingEvaluationMode::Manual)
 	{
 		DetailBuilder.HideProperty(NodeHandle->GetChildHandle(GET_MEMBER_NAME_CHECKED(FAnimNode_OrientationWarping, LocomotionAngle)));
+		DetailBuilder.HideProperty(NodeHandle->GetChildHandle(GET_MEMBER_NAME_CHECKED(FAnimNode_OrientationWarping, LocomotionDirection)));
 		DetailBuilder.HideProperty(NodeHandle->GetChildHandle(GET_MEMBER_NAME_CHECKED(FAnimNode_OrientationWarping, LocomotionAngleDeltaThreshold)));
 		DetailBuilder.HideProperty(NodeHandle->GetChildHandle(GET_MEMBER_NAME_CHECKED(FAnimNode_OrientationWarping, MinRootMotionSpeedThreshold)));
 	}
@@ -145,6 +150,13 @@ void UAnimGraphNode_OrientationWarping::PostEditChangeProperty(struct FPropertyC
 						Pin->BreakAllPinLinks();
 					}
 				}
+				else if (Pin->PinName == GET_MEMBER_NAME_STRING_CHECKED(FAnimNode_OrientationWarping, LocomotionDirection))
+				{
+					if (Node.Mode == EWarpingEvaluationMode::Manual)
+					{
+						Pin->BreakAllPinLinks();
+					}
+				}
 				else if (Pin->PinName == GET_MEMBER_NAME_STRING_CHECKED(FAnimNode_OrientationWarping, LocomotionAngleDeltaThreshold))
 				{
 					if (Node.Mode == EWarpingEvaluationMode::Manual)
diff --git a/Engine/Plugins/Animation/AnimationWarping/Source/Runtime/Private/BoneControllers/AnimNode_OffsetRootBone.cpp b/Engine/Plugins/Animation/AnimationWarping/Source/Runtime/Private/BoneControllers/AnimNode_OffsetRootBone.cpp
--- a/Engine/Plugins/Animation/AnimationWarping/Source/Runtime/Private/BoneControllers/AnimNode_OffsetRootBone.cpp
+++ b/Engine/Plugins/Animation/AnimationWarping/Source/Runtime/Private/BoneControllers/AnimNode_OffsetRootBone.cpp
@@ -5,6 +5,7 @@
 #include "Animation/AnimNodeFunctionRef.h"
 #include "Animation/AnimRootMotionProvider.h"
 #include "HAL/IConsoleManager.h"
+#include "VisualLogger/VisualLogger.h"
 
 #include UE_INLINE_GENERATED_CPP_BY_NAME(AnimNode_OffsetRootBone)
 
@@ -246,6 +247,39 @@ void FAnimNode_OffsetRootBone::EvaluateSkeletalControl_AnyThread(FComponentSpace
 		OutBoneTransforms.Add(FBoneTransform(TargetBoneIndex, TargetBoneTransform));
 	}
 
+#if ENABLE_VISUAL_LOG
+	if (FVisualLogger::IsRecording())
+	{
+		static const TCHAR* LogName = TEXT("OffsetRootBone");
+		const float InnerCircleRadius = 40.0f;
+		const float CircleThickness = 3.0f;
+		const FVector CircleOffset(0,0,1);
+		const float ConeThickness = 0.3f;
+
+		const FTransform TargetBoneInitialTransformWorld = InputBoneTransform * ComponentTransform;
+		const FTransform TargetBoneTransformWorld = TargetBoneTransform * ComponentTransform;
+		UObject* LogOwner = AnimInstanceProxy->GetAnimInstanceObject();
+
+		if (GetMaxTranslationError() >= 0.0f)
+		{
+			const float OuterCircleRadius = GetMaxTranslationError() + InnerCircleRadius;
+			UE_VLOG_CIRCLE_THICK(AnimInstanceProxy->GetAnimInstanceObject(), TEXT("OffsetRootBone"), Display, ComponentTransform.GetLocation() + CircleOffset, FVector::UpVector, OuterCircleRadius, FColor::Red, CircleThickness, TEXT(""));
+		}
+
+		UE_VLOG_CIRCLE_THICK(LogOwner, LogName, Display, ComponentTransform.GetLocation() + CircleOffset, FVector::UpVector, InnerCircleRadius, FColor::Blue, CircleThickness, TEXT(""));
+		UE_VLOG_SEGMENT_THICK(LogOwner, LogName, Display,
+			ComponentTransform.GetLocation() + CircleOffset,
+			ComponentTransform.GetLocation() + InnerCircleRadius * ComponentTransform.GetRotation().GetRightVector() + CircleOffset,
+			FColor::Blue, CircleThickness, TEXT(""));
+		
+		UE_VLOG_CIRCLE_THICK(LogOwner, LogName, Display, TargetBoneTransformWorld.GetLocation() + CircleOffset, FVector::UpVector, InnerCircleRadius, FColor::Green, CircleThickness, TEXT(""));
+		UE_VLOG_SEGMENT_THICK(LogOwner, LogName, Display,
+		 	TargetBoneTransformWorld.GetLocation(),
+		 	TargetBoneTransformWorld.GetLocation() + InnerCircleRadius * TargetBoneTransformWorld.GetRotation().GetRightVector() + CircleOffset,
+			FColor::Green, CircleThickness, TEXT(""));
+	}
+#endif
+
 #if ENABLE_ANIM_DEBUG
 	bool bDebugging = CVarAnimNodeOffsetRootBoneDebug.GetValueOnAnyThread() == 1;
 	if (bDebugging)
diff --git a/Engine/Plugins/Animation/AnimationWarping/Source/Runtime/Private/BoneControllers/AnimNode_OrientationWarping.cpp b/Engine/Plugins/Animation/AnimationWarping/Source/Runtime/Private/BoneControllers/AnimNode_OrientationWarping.cpp
--- a/Engine/Plugins/Animation/AnimationWarping/Source/Runtime/Private/BoneControllers/AnimNode_OrientationWarping.cpp
+++ b/Engine/Plugins/Animation/AnimationWarping/Source/Runtime/Private/BoneControllers/AnimNode_OrientationWarping.cpp
@@ -6,6 +6,8 @@
 #include "Animation/AnimRootMotionProvider.h"
 #include "HAL/IConsoleManager.h"
 #include "Animation/AnimTrace.h"
+#include "Logging/LogVerbosity.h"
+#include "VisualLogger/VisualLogger.h"
 
 #include UE_INLINE_GENERATED_CPP_BY_NAME(AnimNode_OrientationWarping)
 
@@ -176,14 +178,25 @@ void FAnimNode_OrientationWarping::EvaluateSkeletalControl_AnyThread(FComponentS
 			// 2. Actor Velocity
 			// 3. Skeletal Mesh Relative Rotation
 
-			LocomotionAngle = FRotator::NormalizeAxis(LocomotionAngle);
-			// UE-184297 Avoid storing LocomotionAngle in radians in case haven't updated the pinned input, to avoid a DegToRad(RadianValue)
-			const float LocomotionAngleRadians = FMath::DegreesToRadians(LocomotionAngle);
-			const FQuat LocomotionRotation = FQuat(RotationAxisVector, LocomotionAngleRadians);
-
-			const FTransform SkeletalMeshRelativeTransform = Output.AnimInstanceProxy->GetComponentRelativeTransform();
-			const FQuat SkeletalMeshRelativeRotation = SkeletalMeshRelativeTransform.GetRotation();
-			LocomotionForward = SkeletalMeshRelativeRotation.UnrotateVector(LocomotionRotation.GetForwardVector()).GetSafeNormal();
+			
+			if (LocomotionDirection.SquaredLength() > UE_SMALL_NUMBER)
+			{
+				// if we have a LocomotionDirection vector, transform into root bone local space
+				const FCompactPoseBoneIndex RootBoneIndex(0);
+				FTransform RootBoneWorldTransform(Output.Pose.GetComponentSpaceTransform(RootBoneIndex) * Output.AnimInstanceProxy->GetComponentTransform());
+				LocomotionForward = RootBoneWorldTransform.InverseTransformVector(LocomotionDirection);
+				LocomotionForward.Normalize();
+			}
+			else
+			{
+				LocomotionAngle = FRotator::NormalizeAxis(LocomotionAngle);
+				// UE-184297 Avoid storing LocomotionAngle in radians in case haven't updated the pinned input, to avoid a DegToRad(RadianValue)
+				const float LocomotionAngleRadians = FMath::DegreesToRadians(LocomotionAngle);
+				const FQuat LocomotionRotation = FQuat(RotationAxisVector, LocomotionAngleRadians);
+				const FTransform SkeletalMeshRelativeTransform = Output.AnimInstanceProxy->GetComponentRelativeTransform();
+            	const FQuat SkeletalMeshRelativeRotation = SkeletalMeshRelativeTransform.GetRotation();
+				LocomotionForward = SkeletalMeshRelativeRotation.UnrotateVector(LocomotionRotation.GetForwardVector()).GetSafeNormal();
+			}
 
 			// @todo: Graph mode using a "manual value" makes no sense. Restructure logic to address this in the future.
 			if (bUseManualRootMotionVelocity)
@@ -192,7 +205,7 @@ void FAnimNode_OrientationWarping::EvaluateSkeletalControl_AnyThread(FComponentS
 			}
 
 			const FVector RootMotionDeltaTranslation = RootMotionTransformDelta.GetTranslation();
-
+			
 			const float RootMotionDeltaSpeed = RootMotionDeltaTranslation.Size() / DeltaSeconds;
 			if (RootMotionDeltaSpeed < MinRootMotionSpeedThreshold)
 			{
@@ -386,15 +399,59 @@ void FAnimNode_OrientationWarping::EvaluateSkeletalControl_AnyThread(FComponentS
 		}
 	}
 #endif
+	
 
-	const float RootOffset = FMath::UnwindRadians(ActualOrientationAngleRad * DistributedBoneOrientationAlpha);
+#if ENABLE_VISUAL_LOG
+	if (FVisualLogger::IsRecording())
+	{
+		const FTransform ComponentTransform = Output.AnimInstanceProxy->GetComponentTransform();
+		const FVector ActorForwardDirection = Output.AnimInstanceProxy->GetActorTransform().GetRotation().GetForwardVector();
+		FVector DebugArrowOffset = FVector::ZAxisVector * DebugDrawScale;
+
+		// Draw debug shapes
+		{
+			const FVector ForwardDirection = bGraphDrivenWarping
+				? ComponentTransform.GetRotation().RotateVector(LocomotionForward)
+				: ActorForwardDirection;
+
+			UE_VLOG_ARROW(Output.AnimInstanceProxy->GetAnimInstanceObject(), "OrientationWarping", Display,
+				ComponentTransform.GetLocation() + DebugArrowOffset,
+				ComponentTransform.GetLocation() + DebugArrowOffset + ForwardDirection * 100.f * DebugDrawScale,
+				FColor::Red, TEXT(""));
+
+			const FVector RotationDirection = bGraphDrivenWarping
+				? ComponentTransform.GetRotation().RotateVector(RootMotionDeltaDirection)
+				: ActorForwardDirection.RotateAngleAxis(OrientationAngle, RotationAxisVector);
+
+			DebugArrowOffset += FVector::ZAxisVector * DebugDrawScale;
+			UE_VLOG_ARROW(Output.AnimInstanceProxy->GetAnimInstanceObject(), "OrientationWarping", Display,
+				ComponentTransform.GetLocation() + DebugArrowOffset,
+				ComponentTransform.GetLocation() + DebugArrowOffset + RotationDirection * 100.f * DebugDrawScale,
+				FColor::Blue, TEXT(""));
+
+			const float ActualOrientationAngleDegrees = FMath::RadiansToDegrees(ActualOrientationAngleRad);
+			const FVector WarpedRotationDirection = bGraphDrivenWarping
+				? RotationDirection.RotateAngleAxis(ActualOrientationAngleDegrees, RotationAxisVector)
+				: ActorForwardDirection.RotateAngleAxis(ActualOrientationAngleDegrees, RotationAxisVector);
+
+			DebugArrowOffset += FVector::ZAxisVector * DebugDrawScale;
 
+			UE_VLOG_ARROW(Output.AnimInstanceProxy->GetAnimInstanceObject(), "OrientationWarping", Display,
+				ComponentTransform.GetLocation() + DebugArrowOffset,
+				ComponentTransform.GetLocation() + DebugArrowOffset + WarpedRotationDirection * 100.f * DebugDrawScale,
+				FColor::Green, TEXT(""));
+		}
+	}
+#endif
+
+	const float RootOffset = FMath::UnwindRadians(ActualOrientationAngleRad * DistributedBoneOrientationAlpha);
+	
 	// Rotate Root Bone first, as that cheaply rotates the whole pose with one transformation.
 	if (!FMath::IsNearlyZero(RootOffset, KINDA_SMALL_NUMBER))
 	{
 		const FQuat RootRotation = FQuat(RotationAxisVector, RootOffset);
 		const FCompactPoseBoneIndex RootBoneIndex(0);
-
+	
 		FTransform RootBoneTransform(Output.Pose.GetComponentSpaceTransform(RootBoneIndex));
 		RootBoneTransform.SetRotation(RootRotation * RootBoneTransform.GetRotation());
 		RootBoneTransform.NormalizeRotation();
diff --git a/Engine/Plugins/Animation/AnimationWarping/Source/Runtime/Public/BoneControllers/AnimNode_OrientationWarping.h b/Engine/Plugins/Animation/AnimationWarping/Source/Runtime/Public/BoneControllers/AnimNode_OrientationWarping.h
--- a/Engine/Plugins/Animation/AnimationWarping/Source/Runtime/Public/BoneControllers/AnimNode_OrientationWarping.h
+++ b/Engine/Plugins/Animation/AnimationWarping/Source/Runtime/Public/BoneControllers/AnimNode_OrientationWarping.h
@@ -28,6 +28,11 @@ struct ANIMATIONWARPINGRUNTIME_API FAnimNode_OrientationWarping : public FAnimNo
 	UPROPERTY(EditAnywhere, BlueprintReadWrite, Category=Evaluation, meta=(PinShownByDefault))
 	float LocomotionAngle = 0.f;
 
+	// The character movement direction vector in world space
+	// This will be used to compute LocomotionAngle automatically
+	UPROPERTY(EditAnywhere, BlueprintReadWrite, Category=Evaluation, meta=(PinShownByDefault))
+	FVector LocomotionDirection = { 0.f, 0.f, 0.f };
+	
 	// Minimum root motion speed required to apply orientation warping
 	// This is useful to prevent unnatural re-orientation when the animation has a portion with no root motion (i.e starts/stops/idles)
 	// When this value is greater than 0, it's recommended to enable interpolation with RotationInterpSpeed > 0
