Add an option for stride warping to execute even when no root motion data is available

diff --git a/Engine/Plugins/Animation/AnimationWarping/Source/Runtime/Private/BoneControllers/AnimNode_StrideWarping.cpp b/Engine/Plugins/Animation/AnimationWarping/Source/Runtime/Private/BoneControllers/AnimNode_StrideWarping.cpp
--- a/Engine/Plugins/Animation/AnimationWarping/Source/Runtime/Private/BoneControllers/AnimNode_StrideWarping.cpp
+++ b/Engine/Plugins/Animation/AnimationWarping/Source/Runtime/Private/BoneControllers/AnimNode_StrideWarping.cpp
@@ -105,9 +105,9 @@ void FAnimNode_StrideWarping::EvaluateSkeletalControl_AnyThread(FComponentSpaceP
 			bFoundRootMotionAttribute = true;
 #endif
 		}
-		else
+		// Early exit on missing root motion delta attribute
+		else if (bDisableIfMissingRootMotion)
 		{
-			// Early exit on missing root motion delta attribute
 			return;
 		}
 	}
diff --git a/Engine/Plugins/Animation/AnimationWarping/Source/Runtime/Public/BoneControllers/AnimNode_StrideWarping.h b/Engine/Plugins/Animation/AnimationWarping/Source/Runtime/Public/BoneControllers/AnimNode_StrideWarping.h
--- a/Engine/Plugins/Animation/AnimationWarping/Source/Runtime/Public/BoneControllers/AnimNode_StrideWarping.h
+++ b/Engine/Plugins/Animation/AnimationWarping/Source/Runtime/Public/BoneControllers/AnimNode_StrideWarping.h
@@ -107,6 +107,10 @@ struct ANIMATIONWARPINGRUNTIME_API FAnimNode_StrideWarping : public FAnimNode_Sk
 	UPROPERTY(EditAnywhere, Category=Advanced, meta=(DisplayName="Clamp IK Using FK Limits", EditCondition="bCompensateIKUsingFKThighRotation"))
 	bool bClampIKUsingFKLimits = true;
 
+	// Do not execute stride warping if animation data has no root motion
+	UPROPERTY(EditAnywhere, Category=Advanced)
+	bool bDisableIfMissingRootMotion = true;
+
 #if WITH_EDITORONLY_DATA
 	// Scale all debug drawing visualization by a factor
 	UPROPERTY(EditAnywhere, Category=Debug, meta=(ClampMin="0.0"))
