make offset root bone animgraph node operate in local space

diff --git a/Engine/Plugins/Animation/AnimationWarping/Source/Editor/Private/AnimGraph/AnimGraphNode_OffsetRootBone.cpp b/Engine/Plugins/Animation/AnimationWarping/Source/Editor/Private/AnimGraph/AnimGraphNode_OffsetRootBone.cpp
--- a/Engine/Plugins/Animation/AnimationWarping/Source/Editor/Private/AnimGraph/AnimGraphNode_OffsetRootBone.cpp
+++ b/Engine/Plugins/Animation/AnimationWarping/Source/Editor/Private/AnimGraph/AnimGraphNode_OffsetRootBone.cpp
@@ -13,11 +13,6 @@ UAnimGraphNode_OffsetRootBone::UAnimGraphNode_OffsetRootBone(const FObjectInitia
 {
 }
 
-FText UAnimGraphNode_OffsetRootBone::GetControllerDescription() const
-{
-	return LOCTEXT("OffsetRootBone", "Offset Root Bone");
-}
-
 FText UAnimGraphNode_OffsetRootBone::GetTooltipText() const
 {
 	return LOCTEXT("OffsetRootBoneTooltip", "Applies an offset to the root bone");
@@ -25,7 +20,7 @@ FText UAnimGraphNode_OffsetRootBone::GetTooltipText() const
 
 FText UAnimGraphNode_OffsetRootBone::GetNodeTitle(ENodeTitleType::Type TitleType) const
 {
-	return GetControllerDescription();
+	return LOCTEXT("OffsetRootBone", "Offset Root Bone");
 }
 
 FLinearColor UAnimGraphNode_OffsetRootBone::GetNodeTitleColor() const
diff --git a/Engine/Plugins/Animation/AnimationWarping/Source/Editor/Public/AnimGraph/AnimGraphNode_OffsetRootBone.h b/Engine/Plugins/Animation/AnimationWarping/Source/Editor/Public/AnimGraph/AnimGraphNode_OffsetRootBone.h
--- a/Engine/Plugins/Animation/AnimationWarping/Source/Editor/Public/AnimGraph/AnimGraphNode_OffsetRootBone.h
+++ b/Engine/Plugins/Animation/AnimationWarping/Source/Editor/Public/AnimGraph/AnimGraphNode_OffsetRootBone.h
@@ -2,7 +2,7 @@
 
 #pragma once
 
-#include "AnimGraphNode_SkeletalControlBase.h"
+#include "AnimGraphNode_Base.h"
 #include "BoneControllers/AnimNode_OffsetRootBone.h"
 
 #include "AnimGraphNode_OffsetRootBone.generated.h"
@@ -10,10 +10,11 @@
 namespace ENodeTitleType { enum Type : int; }
 
 UCLASS(Experimental)
-class ANIMATIONWARPINGEDITOR_API UAnimGraphNode_OffsetRootBone : public UAnimGraphNode_SkeletalControlBase
+class ANIMATIONWARPINGEDITOR_API UAnimGraphNode_OffsetRootBone : public UAnimGraphNode_Base
 {
 	GENERATED_UCLASS_BODY()
 
+
 	UPROPERTY(EditAnywhere, Category = Settings)
 	FAnimNode_OffsetRootBone Node;
 
@@ -33,8 +34,4 @@ protected:
 	virtual void CustomizePinData(UEdGraphPin* Pin, FName SourcePropertyName, int32 ArrayIndex) const override;
 	// End of UAnimGraphNode_Base interface
 
-	// UAnimGraphNode_SkeletalControlBase interface
-	virtual FText GetControllerDescription() const override;
-	virtual const FAnimNode_SkeletalControlBase* GetNode() const override { return &Node; }
-	// End of UAnimGraphNode_SkeletalControlBase interface
 };
diff --git a/Engine/Plugins/Animation/AnimationWarping/Source/Runtime/Private/BoneControllers/AnimNode_OffsetRootBone.cpp b/Engine/Plugins/Animation/AnimationWarping/Source/Runtime/Private/BoneControllers/AnimNode_OffsetRootBone.cpp
--- a/Engine/Plugins/Animation/AnimationWarping/Source/Runtime/Private/BoneControllers/AnimNode_OffsetRootBone.cpp
+++ b/Engine/Plugins/Animation/AnimationWarping/Source/Runtime/Private/BoneControllers/AnimNode_OffsetRootBone.cpp
@@ -45,19 +45,19 @@ void FAnimNode_OffsetRootBone::GatherDebugData(FNodeDebugData& DebugData)
 	DebugLine += FString::Printf(TEXT("\n - Rotation Halflife: (%.3fd)"), GetRotationHalfLife());
 #endif
 	DebugData.AddDebugItem(DebugLine);
-	ComponentPose.GatherDebugData(DebugData);
 }
 
 void FAnimNode_OffsetRootBone::Initialize_AnyThread(const FAnimationInitializeContext& Context)
 {
-	FAnimNode_SkeletalControlBase::Initialize_AnyThread(Context);
+	FAnimNode_Base::Initialize_AnyThread(Context);
 	AnimInstanceProxy = Context.AnimInstanceProxy;
+	Source.Initialize(Context);
 	Reset(Context);
 }
 
-void FAnimNode_OffsetRootBone::UpdateInternal(const FAnimationUpdateContext& Context)
+void FAnimNode_OffsetRootBone::Update_AnyThread(const FAnimationUpdateContext& Context)
 {
-	FAnimNode_SkeletalControlBase::UpdateInternal(Context);
+	FAnimNode_Base::Update_AnyThread(Context);
 	CachedDeltaTime = Context.GetDeltaTime();
 
 	// If we just became relevant and haven't been initialized yet, then reset.
@@ -66,12 +66,22 @@ void FAnimNode_OffsetRootBone::UpdateInternal(const FAnimationUpdateContext& Con
 		Reset(Context);
 	}
 	UpdateCounter.SynchronizeWith(Context.AnimInstanceProxy->GetUpdateCounter());
+	Source.Update(Context);
 }
 
-void FAnimNode_OffsetRootBone::EvaluateSkeletalControl_AnyThread(FComponentSpacePoseContext& Output, TArray<FBoneTransform>& OutBoneTransforms)
+void FAnimNode_OffsetRootBone::Evaluate_AnyThread(FPoseContext& Output)
 {
 	SCOPE_CYCLE_COUNTER(STAT_OffsetRootBone_Eval);
-	check(OutBoneTransforms.Num() == 0);
+
+	Super::Evaluate_AnyThread(Output);
+	Source.Evaluate(Output);
+
+#if ENABLE_ANIM_DEBUG
+	if (CVarAnimNodeOffsetRootBoneEnable.GetValueOnAnyThread() == 0)
+	{
+		return;
+	}
+#endif
 
 	bool bGraphDriven = false;
 	const UE::Anim::IAnimRootMotionProvider* RootMotionProvider = UE::Anim::IAnimRootMotionProvider::Get();
@@ -83,7 +93,7 @@ void FAnimNode_OffsetRootBone::EvaluateSkeletalControl_AnyThread(FComponentSpace
 	}
 
 	const FCompactPoseBoneIndex TargetBoneIndex(0);
-	const FTransform InputBoneTransform = Output.Pose.GetComponentSpaceTransform(TargetBoneIndex);
+	const FTransform InputBoneTransform = Output.Pose[TargetBoneIndex];
 
 	const FTransform LastComponentTransform = ComponentTransform;
 	ComponentTransform = AnimInstanceProxy->GetComponentTransform();
@@ -144,7 +154,7 @@ void FAnimNode_OffsetRootBone::EvaluateSkeletalControl_AnyThread(FComponentSpace
 	// Simulated translation should stay the same along the approach direction
 	SimulatedTranslation = FVector::PointPlaneProject(SimulatedTranslation, ComponentTransform.GetLocation(), GravityDirCS);
 
-	const FBoneContainer& RequiredBones = Output.Pose.GetPose().GetBoneContainer();
+	const FBoneContainer& RequiredBones = Output.Pose.GetBoneContainer();
 
 	bool bModifyBone = true;
 #if ENABLE_ANIM_DEBUG
@@ -242,10 +252,8 @@ void FAnimNode_OffsetRootBone::EvaluateSkeletalControl_AnyThread(FComponentSpace
 	FTransform TargetBoneTransform = SimulatedTransform * ComponentTransform.Inverse();
 	// Accumulate the input bone transform to keep the offset independent from any previous adjustments to the root
 	TargetBoneTransform.Accumulate(InputBoneTransform);
-	if (bModifyBone)
-	{
-		OutBoneTransforms.Add(FBoneTransform(TargetBoneIndex, TargetBoneTransform));
-	}
+
+	Output.Pose[TargetBoneIndex] = TargetBoneTransform;
 
 #if ENABLE_VISUAL_LOG
 	if (FVisualLogger::IsRecording())
@@ -333,18 +341,6 @@ void FAnimNode_OffsetRootBone::EvaluateSkeletalControl_AnyThread(FComponentSpace
 	bIsFirstUpdate = false;
 }
 
-bool FAnimNode_OffsetRootBone::IsValidToEvaluate(const USkeleton* Skeleton, const FBoneContainer& RequiredBones)
-{
-#if ENABLE_ANIM_DEBUG
-	if (CVarAnimNodeOffsetRootBoneEnable.GetValueOnAnyThread() == 0)
-	{
-		return false;
-	}
-#endif
-
-	return true;
-}
-
 EWarpingEvaluationMode FAnimNode_OffsetRootBone::GetEvaluationMode() const
 {
 	return GET_ANIM_NODE_DATA(EWarpingEvaluationMode, EvaluationMode);
diff --git a/Engine/Plugins/Animation/AnimationWarping/Source/Runtime/Public/BoneControllers/AnimNode_OffsetRootBone.h b/Engine/Plugins/Animation/AnimationWarping/Source/Runtime/Public/BoneControllers/AnimNode_OffsetRootBone.h
--- a/Engine/Plugins/Animation/AnimationWarping/Source/Runtime/Public/BoneControllers/AnimNode_OffsetRootBone.h
+++ b/Engine/Plugins/Animation/AnimationWarping/Source/Runtime/Public/BoneControllers/AnimNode_OffsetRootBone.h
@@ -31,10 +31,13 @@ enum class EOffsetRootBoneMode : uint8
 };
 
 USTRUCT(BlueprintInternalUseOnly, Experimental)
-struct ANIMATIONWARPINGRUNTIME_API FAnimNode_OffsetRootBone : public FAnimNode_SkeletalControlBase
+struct ANIMATIONWARPINGRUNTIME_API FAnimNode_OffsetRootBone : public FAnimNode_Base
 {
 	GENERATED_BODY();
 
+	UPROPERTY(EditAnywhere, BlueprintReadWrite, Category = Links, meta = (DisplayPriority = 0))
+	FPoseLink Source;
+
 #if WITH_EDITORONLY_DATA
 	UPROPERTY(EditAnywhere, Category = Evaluation, meta=(FoldProperty))
 	EWarpingEvaluationMode EvaluationMode = EWarpingEvaluationMode::Graph;
@@ -101,14 +104,10 @@ struct ANIMATIONWARPINGRUNTIME_API FAnimNode_OffsetRootBone : public FAnimNode_S
 public:
 	// FAnimNode_Base interface
 	virtual void GatherDebugData(FNodeDebugData& DebugData) override;
-	virtual void UpdateInternal(const FAnimationUpdateContext& Context) override;
-	// End of FAnimNode_Base interface
-
-	// FAnimNode_SkeletalControlBase interface
+	virtual void Update_AnyThread(const FAnimationUpdateContext& Context) override;
+	virtual void Evaluate_AnyThread(FPoseContext& Output) override;
 	virtual void Initialize_AnyThread(const FAnimationInitializeContext& Context) override;
-	virtual void EvaluateSkeletalControl_AnyThread(FComponentSpacePoseContext& Output, TArray<FBoneTransform>& OutBoneTransforms) override;
-	virtual bool IsValidToEvaluate(const USkeleton* Skeleton, const FBoneContainer& RequiredBones) override;
-	// End of FAnimNode_SkeletalControlBase interface
+	// End of FAnimNode_Base interface
 
 	// Folded property accesors
 	EWarpingEvaluationMode GetEvaluationMode() const;
