Steering animgraph node - makes procedural modifications to root motion attribute to reach a desired facing direction

diff --git a/Engine/Plugins/Animation/AnimationWarping/Source/Editor/Private/AnimGraph/AnimGraphNode_Steering.cpp b/Engine/Plugins/Animation/AnimationWarping/Source/Editor/Private/AnimGraph/AnimGraphNode_Steering.cpp
new file mode 100644
--- /dev/null
+++ b/Engine/Plugins/Animation/AnimationWarping/Source/Editor/Private/AnimGraph/AnimGraphNode_Steering.cpp
@@ -0,0 +1,70 @@
+// Copyright Epic Games, Inc. All Rights Reserved.
+
+#include "AnimGraph/AnimGraphNode_Steering.h"
+#include "Animation/AnimRootMotionProvider.h"
+#include "Kismet2/BlueprintEditorUtils.h"
+#include "Kismet2/CompilerResultsLog.h"
+#include "DetailCategoryBuilder.h"
+#include "DetailLayoutBuilder.h"
+#include "PropertyHandle.h"
+#include "ScopedTransaction.h"
+
+#define LOCTEXT_NAMESPACE "AnimationWarping"
+
+UAnimGraphNode_Steering::UAnimGraphNode_Steering(const FObjectInitializer& ObjectInitializer) : Super(ObjectInitializer)
+{
+}
+
+FText UAnimGraphNode_Steering::GetControllerDescription() const
+{
+	return LOCTEXT("Steering", "Steering");
+}
+
+FText UAnimGraphNode_Steering::GetTooltipText() const
+{
+	return LOCTEXT("SteeringTooltip", "Rotates the root and lower body by the specified angle, while counter rotating the upper body to maintain the forward facing direction.");
+}
+
+FText UAnimGraphNode_Steering::GetNodeTitle(ENodeTitleType::Type TitleType) const
+{
+	return GetControllerDescription();
+}
+
+FLinearColor UAnimGraphNode_Steering::GetNodeTitleColor() const
+{
+	return FLinearColor(FColor(153.f, 0.f, 0.f));
+}
+
+void UAnimGraphNode_Steering::CustomizePinData(UEdGraphPin* Pin, FName SourcePropertyName, int32 ArrayIndex) const
+{
+	Super::CustomizePinData(Pin, SourcePropertyName, ArrayIndex);
+}
+
+void UAnimGraphNode_Steering::CustomizeDetails(IDetailLayoutBuilder& DetailBuilder)
+{
+	DECLARE_SCOPE_HIERARCHICAL_COUNTER_FUNC()
+	Super::CustomizeDetails(DetailBuilder);
+}
+
+void UAnimGraphNode_Steering::PostEditChangeProperty(struct FPropertyChangedEvent& PropertyChangedEvent)
+{
+	DECLARE_SCOPE_HIERARCHICAL_COUNTER_FUNC()
+	Super::PostEditChangeProperty(PropertyChangedEvent);
+}
+
+void UAnimGraphNode_Steering::GetInputLinkAttributes(FNodeAttributeArray& OutAttributes) const
+{
+	OutAttributes.Add(UE::Anim::IAnimRootMotionProvider::AttributeName);
+}
+
+void UAnimGraphNode_Steering::GetOutputLinkAttributes(FNodeAttributeArray& OutAttributes) const
+{
+	OutAttributes.Add(UE::Anim::IAnimRootMotionProvider::AttributeName);
+}
+
+void UAnimGraphNode_Steering::ValidateAnimNodeDuringCompilation(USkeleton* ForSkeleton, FCompilerResultsLog& MessageLog)
+{
+	Super::ValidateAnimNodeDuringCompilation(ForSkeleton, MessageLog);
+}
+
+#undef LOCTEXT_NAMESPACE
diff --git a/Engine/Plugins/Animation/AnimationWarping/Source/Editor/Public/AnimGraph/AnimGraphNode_Steering.h b/Engine/Plugins/Animation/AnimationWarping/Source/Editor/Public/AnimGraph/AnimGraphNode_Steering.h
new file mode 100644
--- /dev/null
+++ b/Engine/Plugins/Animation/AnimationWarping/Source/Editor/Public/AnimGraph/AnimGraphNode_Steering.h
@@ -0,0 +1,39 @@
+// Copyright Epic Games, Inc. All Rights Reserved.
+
+#pragma once
+#include "AnimGraphNode_SkeletalControlBase.h"
+#include "BoneControllers/AnimNode_Steering.h"
+#include "AnimGraphNode_Steering.generated.h"
+
+namespace ENodeTitleType { enum Type : int; }
+
+UCLASS()
+class ANIMATIONWARPINGEDITOR_API UAnimGraphNode_Steering : public UAnimGraphNode_SkeletalControlBase
+{
+	GENERATED_UCLASS_BODY()
+
+	UPROPERTY(EditAnywhere, Category = Settings)
+	FAnimNode_Steering Node;
+
+public:
+	// UEdGraphNode interface
+	virtual FText GetNodeTitle(ENodeTitleType::Type TitleType) const override;
+	virtual FLinearColor GetNodeTitleColor() const override;
+	virtual FText GetTooltipText() const override;
+	virtual void CustomizeDetails(IDetailLayoutBuilder& DetailBuilder) override;
+	virtual void PostEditChangeProperty(struct FPropertyChangedEvent& PropertyChangedEvent) override;
+	virtual void GetInputLinkAttributes(FNodeAttributeArray& OutAttributes) const override;
+	virtual void GetOutputLinkAttributes(FNodeAttributeArray& OutAttributes) const override;
+	virtual void ValidateAnimNodeDuringCompilation(USkeleton* ForSkeleton, FCompilerResultsLog& MessageLog) override;
+	// End of UEdGraphNode interface
+
+protected:
+	// UAnimGraphNode_Base interface
+	virtual void CustomizePinData(UEdGraphPin* Pin, FName SourcePropertyName, int32 ArrayIndex) const override;
+	// End of UAnimGraphNode_Base interface
+
+	// UAnimGraphNode_SkeletalControlBase interface
+	virtual FText GetControllerDescription() const override;
+	virtual const FAnimNode_SkeletalControlBase* GetNode() const override { return &Node; }
+	// End of UAnimGraphNode_SkeletalControlBase interface
+};
diff --git a/Engine/Plugins/Animation/AnimationWarping/Source/Runtime/Private/BoneControllers/AnimNode_Steering.cpp b/Engine/Plugins/Animation/AnimationWarping/Source/Runtime/Private/BoneControllers/AnimNode_Steering.cpp
new file mode 100644
--- /dev/null
+++ b/Engine/Plugins/Animation/AnimationWarping/Source/Runtime/Private/BoneControllers/AnimNode_Steering.cpp
@@ -0,0 +1,101 @@
+// Copyright Epic Games, Inc. All Rights Reserved.
+
+#include "BoneControllers/AnimNode_Steering.h"
+#include "Animation/AnimInstanceProxy.h"
+#include "Animation/AnimNodeFunctionRef.h"
+#include "Animation/AnimRootMotionProvider.h"
+#include "HAL/IConsoleManager.h"
+#include "Animation/AnimTrace.h"
+#include "Logging/LogVerbosity.h"
+#include "VisualLogger/VisualLogger.h"
+
+#include UE_INLINE_GENERATED_CPP_BY_NAME(AnimNode_Steering)
+
+void FAnimNode_Steering::UpdateInternal(const FAnimationUpdateContext& Context)
+{
+	FAnimNode_SkeletalControlBase::UpdateInternal(Context);
+
+	if (USkeletalMeshComponent* SkelMeshComponent = Context.AnimInstanceProxy->GetSkelMeshComponent())
+	{
+		RootBoneTransform = SkelMeshComponent->GetBoneTransform(0);
+	}
+	else
+	{
+		RootBoneTransform = Context.AnimInstanceProxy->GetComponentTransform();
+	}
+}
+
+void FAnimNode_Steering::Initialize_AnyThread(const FAnimationInitializeContext& Context)
+{
+	FAnimNode_SkeletalControlBase::Initialize_AnyThread(Context);
+}
+
+void FAnimNode_Steering::EvaluateSkeletalControl_AnyThread(FComponentSpacePoseContext& Output, TArray<FBoneTransform>& OutBoneTransforms)
+{
+	if (Alpha > 0.0f)
+	{
+		const UE::Anim::IAnimRootMotionProvider* RootMotionProvider = UE::Anim::IAnimRootMotionProvider::Get();
+		ensureMsgf(RootMotionProvider, TEXT("Steering expected a valid root motion delta provider interface."));
+
+		if (RootMotionProvider)
+		{
+			FTransform RootMotionTransformDelta = FTransform::Identity;
+			if (RootMotionProvider->ExtractRootMotion(Output.CustomAttributes, RootMotionTransformDelta))
+			{
+				const float DeltaSeconds = Output.AnimInstanceProxy->GetDeltaSeconds();
+
+				FQuat RootBoneRotation = RootBoneTransform.GetRotation();
+
+				UE_VLOG_ARROW(Output.AnimInstanceProxy->GetAnimInstanceObject(), "Steering", Display,
+					RootBoneTransform.GetLocation(),
+					RootBoneTransform.GetLocation()  + RootBoneRotation.GetForwardVector() * 100,
+					FColor::Red, TEXT(""));
+					
+				UE_VLOG_ARROW(Output.AnimInstanceProxy->GetAnimInstanceObject(), "Steering", Display,
+					RootBoneTransform.GetLocation(),
+					RootBoneTransform.GetLocation()  + TargetOrientation.GetForwardVector() * 100,
+					FColor::Green, TEXT(""));
+
+				FQuat Delta =  TargetOrientation * RootBoneRotation.Inverse();
+
+
+				if (TargetTime > 0)
+				{
+					if (UAnimSequenceBase* AnimSequence = Cast<UAnimSequenceBase>(CurrentAnimAsset))
+					{
+						FTransform RootMotionDelta = AnimSequence->ExtractRootMotion(CurrentAnimAssetTime, TargetTime, true);
+						FQuat RootMotionRotation = RootMotionDelta.GetRotation();
+
+						FRotator RootMotionRot(RootMotionRotation);
+						
+						if (fabs(RootMotionRot.Yaw) > RootMotionThreshold)
+						{
+							UE_VLOG_ARROW(Output.AnimInstanceProxy->GetAnimInstanceObject(), "Steering", Display,
+								RootBoneTransform.GetLocation(),
+								RootBoneTransform.GetLocation()  + (RootMotionRotation * RootBoneRotation).GetForwardVector() * 100,
+								FColor::Blue, TEXT(""));
+
+							FRotator DeltaRot(Delta);
+							
+							float Ratio =  (TargetTime / DeltaSeconds) * DeltaRot.Yaw / RootMotionRot.Yaw ;
+
+							FRotator RootMotionFrameRot(RootMotionTransformDelta.GetRotation());
+
+							RootMotionFrameRot.Yaw *= Ratio;
+							
+							RootMotionTransformDelta.SetRotation(FQuat(RootMotionFrameRot));
+							RootMotionProvider->OverrideRootMotion(RootMotionTransformDelta, Output.CustomAttributes);
+							
+							return;
+						}
+					}
+					
+					Delta = FQuat::Slerp(FQuat::Identity, Delta,  DeltaSeconds/TargetTime);
+				}
+
+				RootMotionTransformDelta.SetRotation(FQuat::Slerp(RootMotionTransformDelta.GetRotation(), Delta, Alpha));
+				RootMotionProvider->OverrideRootMotion(RootMotionTransformDelta, Output.CustomAttributes);
+			}
+		}
+	}
+}
\ No newline at end of file
diff --git a/Engine/Plugins/Animation/AnimationWarping/Source/Runtime/Public/BoneControllers/AnimNode_Steering.h b/Engine/Plugins/Animation/AnimationWarping/Source/Runtime/Public/BoneControllers/AnimNode_Steering.h
new file mode 100644
--- /dev/null
+++ b/Engine/Plugins/Animation/AnimationWarping/Source/Runtime/Public/BoneControllers/AnimNode_Steering.h
@@ -0,0 +1,55 @@
+// Copyright Epic Games, Inc. All Rights Reserved.
+
+#pragma once
+
+#include "BoneControllers/BoneControllerTypes.h"
+#include "BoneControllers/AnimNode_SkeletalControlBase.h"
+#include "AnimNode_Steering.generated.h"
+
+struct FAnimationInitializeContext;
+struct FComponentSpacePoseContext;
+struct FNodeDebugData;
+
+// add procedural delta to the root motion attribute 
+USTRUCT(BlueprintInternalUseOnly)
+struct ANIMATIONWARPINGRUNTIME_API FAnimNode_Steering : public FAnimNode_SkeletalControlBase
+{
+	GENERATED_BODY()
+
+	// The Orientation to steer towards
+	UPROPERTY(EditAnywhere, BlueprintReadWrite, Category=Evaluation, meta=(PinShownByDefault))
+	FQuat TargetOrientation = FQuat::Identity;
+	
+	// The number of seconds in the future before we should reach the TargetOrientation
+	UPROPERTY(EditAnywhere, BlueprintReadWrite, Category=Evaluation, meta=(PinShownByDefault))
+	float TargetTime = 0.2f;
+
+	// If less than this number of degrees of rotation is found over TargetTime seconds in the CurrentAnimAsset, then rotation will be added linearly to reach the TargetOrientation
+	// Otherwise the root motion will be scaled to reach the TargetOrientation over TargetTime seconds
+	UPROPERTY(EditAnywhere, Category=Evaluation)
+	float RootMotionThreshold = 1.0f;
+	
+	
+	
+	// Animation Asset for incorporating root motion data. If CurrentAnimAsset is set, and the animation has root motion rotation within the TargetTime, then those rotations will be scaled to reach the TargetOrientation
+	UPROPERTY(EditAnywhere, BlueprintReadWrite, Category=Evaluation, meta=(PinShownByDefault))
+	TObjectPtr<UAnimationAsset> CurrentAnimAsset;
+	
+	// Current playback time in seconds of the CurrentAnimAsset
+	UPROPERTY(EditAnywhere, BlueprintReadWrite, Category=Evaluation, meta=(PinShownByDefault))
+	float CurrentAnimAssetTime = 0.f;
+
+	// FAnimNodeBase interface
+	virtual void Initialize_AnyThread(const FAnimationInitializeContext& Context) override;
+	virtual void UpdateInternal(const FAnimationUpdateContext& Context) override;
+	// End of FAnimNodeBase interface
+	
+	// FAnimNode_SkeletalControlBase interface
+	virtual void EvaluateSkeletalControl_AnyThread(FComponentSpacePoseContext& Output, TArray<FBoneTransform>& OutBoneTransforms) override;
+	virtual bool IsValidToEvaluate(const USkeleton* Skeleton, const FBoneContainer& RequiredBones) override { return true; }
+	// End of FAnimNode_SkeletalControlBase interface
+
+private:
+
+	FTransform RootBoneTransform;
+};
diff --git a/Engine/Plugins/Experimental/Animation/BlendStack/Source/Runtime/Private/BlendStackAnimNodeLibrary.cpp b/Engine/Plugins/Experimental/Animation/BlendStack/Source/Runtime/Private/BlendStackAnimNodeLibrary.cpp
--- a/Engine/Plugins/Experimental/Animation/BlendStack/Source/Runtime/Private/BlendStackAnimNodeLibrary.cpp
+++ b/Engine/Plugins/Experimental/Animation/BlendStack/Source/Runtime/Private/BlendStackAnimNodeLibrary.cpp
@@ -2,6 +2,7 @@
 
 #include "BlendStack/BlendStackAnimNodeLibrary.h"
 #include "BlendStack/AnimNode_BlendStack.h"
+#include "BlendStack/AnimNode_BlendStackInput.h"
 #include "BlendStack/BlendStackDefines.h"
 
 FBlendStackAnimNodeReference UBlendStackAnimNodeLibrary::ConvertToBlendStackNode(const FAnimNodeReference& Node, EAnimNodeReferenceConversionResult& Result)
@@ -9,6 +10,30 @@ FBlendStackAnimNodeReference UBlendStackAnimNodeLibrary::ConvertToBlendStackNode
 	return FAnimNodeReference::ConvertToType<FBlendStackAnimNodeReference>(Node, Result);
 }
 
+UAnimationAsset* UBlendStackAnimNodeLibrary::GetCurrentBlendStackAnimAsset(const FAnimNodeReference& Node)
+{
+	if (FAnimNode_BlendStackInput* BlendStackInput = Node.GetAnimNodePtr<FAnimNode_BlendStackInput>())
+	{
+		if (BlendStackInput->Player && *BlendStackInput->Player )
+		{
+			return (*BlendStackInput->Player)->GetAnimationAsset();
+		}
+	}
+	return nullptr;
+}
+
+float UBlendStackAnimNodeLibrary::GetCurrentBlendStackAnimAssetTime(const FAnimNodeReference& Node)
+{
+	if (FAnimNode_BlendStackInput* BlendStackInput = Node.GetAnimNodePtr<FAnimNode_BlendStackInput>())
+	{
+		if (BlendStackInput->Player && *BlendStackInput->Player )
+		{
+			return (*BlendStackInput->Player)->GetAccumulatedTime();
+		}
+	}
+	return 0.f;
+}
+
 void UBlendStackAnimNodeLibrary::BlendTo(const FAnimUpdateContext& Context, 
 										const FBlendStackAnimNodeReference& BlendStackNode, 
 										UAnimationAsset* AnimationAsset,
diff --git a/Engine/Plugins/Experimental/Animation/BlendStack/Source/Runtime/Public/BlendStack/BlendStackAnimNodeLibrary.h b/Engine/Plugins/Experimental/Animation/BlendStack/Source/Runtime/Public/BlendStack/BlendStackAnimNodeLibrary.h
--- a/Engine/Plugins/Experimental/Animation/BlendStack/Source/Runtime/Public/BlendStack/BlendStackAnimNodeLibrary.h
+++ b/Engine/Plugins/Experimental/Animation/BlendStack/Source/Runtime/Public/BlendStack/BlendStackAnimNodeLibrary.h
@@ -27,6 +27,15 @@ public:
 	/** Get a blend stack node context from an anim node context */
 	UFUNCTION(BlueprintCallable, Category = "Animation|BlendStack", meta = (BlueprintThreadSafe, ExpandEnumAsExecs = "Result"))
 	static FBlendStackAnimNodeReference ConvertToBlendStackNode(const FAnimNodeReference& Node, EAnimNodeReferenceConversionResult& Result);
+	
+	
+	/** Get the current AnimationAsset that is playing from a Blend Stack Input node */
+    UFUNCTION(BlueprintPure, Category = "Animation|BlendStack", meta = (BlueprintThreadSafe))
+    static UAnimationAsset* GetCurrentBlendStackAnimAsset(const FAnimNodeReference& Node);
+	
+	/** Get the current elapsed time of the animation that is playing from a Blend Stack Input node */
+    UFUNCTION(BlueprintPure, Category = "Animation|BlendStack", meta = (BlueprintThreadSafe))
+    static float GetCurrentBlendStackAnimAssetTime(const FAnimNodeReference& Node);
 
 	/** Get a blend stack node context from an anim node context (pure) */
 	UFUNCTION(BlueprintPure, Category = "Animation|BlendStack", meta = (BlueprintThreadSafe, DisplayName = "Convert to Blend Stack Node"))
diff --git a/Engine/Plugins/Experimental/Animation/PoseSearch/Source/Runtime/Public/PoseSearch/PoseSearchTrajectoryLibrary.h b/Engine/Plugins/Experimental/Animation/PoseSearch/Source/Runtime/Public/PoseSearch/PoseSearchTrajectoryLibrary.h
--- a/Engine/Plugins/Experimental/Animation/PoseSearch/Source/Runtime/Public/PoseSearch/PoseSearchTrajectoryLibrary.h
+++ b/Engine/Plugins/Experimental/Animation/PoseSearch/Source/Runtime/Public/PoseSearch/PoseSearchTrajectoryLibrary.h
@@ -111,7 +111,7 @@ public:
 	static void HandleTrajectoryWorldCollisions(const UObject* WorldContextObject, UPARAM(ref) const FPoseSearchQueryTrajectory& InTrajectory, FPoseSearchQueryTrajectory& OutTrajectory,
 		ETraceTypeQuery TraceChannel, bool bTraceComplex, const TArray<AActor*>& ActorsToIgnore, EDrawDebugTrace::Type DrawDebugType, bool bIgnoreSelf, FLinearColor TraceColor = FLinearColor::Red, FLinearColor TraceHitColor = FLinearColor::Green, float DrawTime = 5.0f);
 
-	UFUNCTION(BlueprintPure, Category="Animation|PoseSearch")
+	UFUNCTION(BlueprintPure, meta = (BlueprintThreadSafe), Category="Animation|PoseSearch")
 	static void GetTrajectorySampleAtTime(UPARAM(ref) const FPoseSearchQueryTrajectory& InTrajectory, float Time, FPoseSearchQueryTrajectorySample& OutTrajectorySample, bool bExtrapolate = true);
 
 private:
