Make Orientation Warping node operate relative to the Root Bone Transform from the previous frame, to work better with  Offset Root Bone

diff --git a/Engine/Plugins/Animation/AnimationWarping/Source/Editor/Private/AnimGraph/AnimGraphNode_OrientationWarping.cpp b/Engine/Plugins/Animation/AnimationWarping/Source/Editor/Private/AnimGraph/AnimGraphNode_OrientationWarping.cpp
--- a/Engine/Plugins/Animation/AnimationWarping/Source/Editor/Private/AnimGraph/AnimGraphNode_OrientationWarping.cpp
+++ b/Engine/Plugins/Animation/AnimationWarping/Source/Editor/Private/AnimGraph/AnimGraphNode_OrientationWarping.cpp
@@ -57,6 +57,11 @@ void UAnimGraphNode_OrientationWarping::CustomizePinData(UEdGraphPin* Pin, FName
 	{
 		Pin->bHidden = (Node.Mode == EWarpingEvaluationMode::Manual);
 	}
+
+	if (Pin->PinName == GET_MEMBER_NAME_STRING_CHECKED(FAnimNode_OrientationWarping, WarpingSpace))
+	{
+		Pin->bHidden = (Node.WarpingSpace == EOrientationWarpingSpace::CustomTransform);
+	}
 }
 
 void UAnimGraphNode_OrientationWarping::CustomizeDetails(IDetailLayoutBuilder& DetailBuilder)
@@ -114,6 +119,11 @@ void UAnimGraphNode_OrientationWarping::CustomizeDetails(IDetailLayoutBuilder& D
 		DetailBuilder.HideProperty(NodeHandle->GetChildHandle(GET_MEMBER_NAME_CHECKED(FAnimNode_OrientationWarping, LocomotionAngleDeltaThreshold)));
 		DetailBuilder.HideProperty(NodeHandle->GetChildHandle(GET_MEMBER_NAME_CHECKED(FAnimNode_OrientationWarping, MinRootMotionSpeedThreshold)));
 	}
+	
+	if (Node.WarpingSpace != EOrientationWarpingSpace::CustomTransform)
+   	{
+    	DetailBuilder.HideProperty(NodeHandle->GetChildHandle(GET_MEMBER_NAME_CHECKED(FAnimNode_OrientationWarping, WarpingSpaceTransform)));
+    }
 }
 
 void UAnimGraphNode_OrientationWarping::PostEditChangeProperty(struct FPropertyChangedEvent& PropertyChangedEvent)
diff --git a/Engine/Plugins/Animation/AnimationWarping/Source/Runtime/Private/BoneControllers/AnimNode_OrientationWarping.cpp b/Engine/Plugins/Animation/AnimationWarping/Source/Runtime/Private/BoneControllers/AnimNode_OrientationWarping.cpp
--- a/Engine/Plugins/Animation/AnimationWarping/Source/Runtime/Private/BoneControllers/AnimNode_OrientationWarping.cpp
+++ b/Engine/Plugins/Animation/AnimationWarping/Source/Runtime/Private/BoneControllers/AnimNode_OrientationWarping.cpp
@@ -1,6 +1,7 @@
 // Copyright Epic Games, Inc. All Rights Reserved.
 
 #include "BoneControllers/AnimNode_OrientationWarping.h"
+
 #include "Animation/AnimInstanceProxy.h"
 #include "Animation/AnimNodeFunctionRef.h"
 #include "Animation/AnimRootMotionProvider.h"
@@ -106,6 +107,18 @@ void FAnimNode_OrientationWarping::UpdateInternal(const FAnimationUpdateContext&
 	}
 	UpdateCounter.SynchronizeWith(Context.AnimInstanceProxy->GetUpdateCounter());
 	BlendWeight = Context.GetFinalBlendWeight();
+
+	if (WarpingSpace == EOrientationWarpingSpace::RootBoneTransform)
+	{
+		if (USkeletalMeshComponent* SkelMeshComponent = Context.AnimInstanceProxy->GetSkelMeshComponent())
+		{
+			WarpingSpaceTransform = SkelMeshComponent->GetBoneTransform(0);
+		}
+		else
+		{
+			WarpingSpaceTransform = Context.AnimInstanceProxy->GetComponentTransform();
+		}
+	}
 }
 
 void FAnimNode_OrientationWarping::EvaluateSkeletalControl_AnyThread(FComponentSpacePoseContext& Output, TArray<FBoneTransform>& OutBoneTransforms)
@@ -153,6 +166,11 @@ void FAnimNode_OrientationWarping::EvaluateSkeletalControl_AnyThread(FComponentS
 	// speed information about the motion. It may also allow us to do more complex orienting behavior 
 	// when multiple degrees of freedom can be considered.
 
+	if (WarpingSpace == EOrientationWarpingSpace::ComponentTransform)
+	{
+		WarpingSpaceTransform = Output.AnimInstanceProxy->GetComponentTransform();
+	}
+
 	if (bGraphDrivenWarping)
 	{
 #if !ENABLE_ANIM_DEBUG
@@ -182,9 +200,7 @@ void FAnimNode_OrientationWarping::EvaluateSkeletalControl_AnyThread(FComponentS
 			if (LocomotionDirection.SquaredLength() > UE_SMALL_NUMBER)
 			{
 				// if we have a LocomotionDirection vector, transform into root bone local space
-				const FCompactPoseBoneIndex RootBoneIndex(0);
-				FTransform RootBoneWorldTransform(Output.Pose.GetComponentSpaceTransform(RootBoneIndex) * Output.AnimInstanceProxy->GetComponentTransform());
-				LocomotionForward = RootBoneWorldTransform.InverseTransformVector(LocomotionDirection);
+				LocomotionForward = WarpingSpaceTransform.InverseTransformVector(LocomotionDirection);
 				LocomotionForward.Normalize();
 			}
 			else
diff --git a/Engine/Plugins/Animation/AnimationWarping/Source/Runtime/Public/BoneControllers/AnimNode_OrientationWarping.h b/Engine/Plugins/Animation/AnimationWarping/Source/Runtime/Public/BoneControllers/AnimNode_OrientationWarping.h
--- a/Engine/Plugins/Animation/AnimationWarping/Source/Runtime/Public/BoneControllers/AnimNode_OrientationWarping.h
+++ b/Engine/Plugins/Animation/AnimationWarping/Source/Runtime/Public/BoneControllers/AnimNode_OrientationWarping.h
@@ -10,6 +10,18 @@ struct FAnimationInitializeContext;
 struct FComponentSpacePoseContext;
 struct FNodeDebugData;
 
+
+UENUM(BlueprintType)
+enum class EOrientationWarpingSpace : uint8
+{
+	// apply warping relative to current component transform
+	ComponentTransform,
+	// Apply warping relative to previous frame's root bone transform. Use this mode when using an OffsetRootBone node which allows the root bone and component transforms to differ. 
+	RootBoneTransform,
+	// Provide a custom transform pin
+	CustomTransform
+};
+
 USTRUCT(BlueprintInternalUseOnly)
 struct ANIMATIONWARPINGRUNTIME_API FAnimNode_OrientationWarping : public FAnimNode_SkeletalControlBase
 {
@@ -100,6 +112,13 @@ struct ANIMATIONWARPINGRUNTIME_API FAnimNode_OrientationWarping : public FAnimNo
 	UPROPERTY(EditAnywhere, BlueprintReadWrite, Category=Experimental, meta=(PinHiddenByDefault))
 	FVector ManualRootMotionVelocity = FVector::ZeroVector;
 
+	
+	UPROPERTY(EditAnywhere, BlueprintReadWrite, Category=Evaluation, meta=(PinHiddenByDefault))
+	EOrientationWarpingSpace WarpingSpace;
+	
+	UPROPERTY(EditAnywhere, BlueprintReadWrite, Category=Evaluation, meta=(PinHiddenByDefault))
+	FTransform WarpingSpaceTransform;
+
 #if WITH_EDITORONLY_DATA
 	// Scale all debug drawing visualization by a factor
 	UPROPERTY(EditAnywhere, Category=Debug, meta=(ClampMin="0.0"))
