AnimationWarping - implementing FAnimNode_OffsetRootBone::CacheBones_AnyThread

diff --git a/Engine/Plugins/Animation/AnimationWarping/Source/Runtime/Private/BoneControllers/AnimNode_OffsetRootBone.cpp b/Engine/Plugins/Animation/AnimationWarping/Source/Runtime/Private/BoneControllers/AnimNode_OffsetRootBone.cpp
--- a/Engine/Plugins/Animation/AnimationWarping/Source/Runtime/Private/BoneControllers/AnimNode_OffsetRootBone.cpp
+++ b/Engine/Plugins/Animation/AnimationWarping/Source/Runtime/Private/BoneControllers/AnimNode_OffsetRootBone.cpp
@@ -49,15 +49,21 @@ void FAnimNode_OffsetRootBone::GatherDebugData(FNodeDebugData& DebugData)
 
 void FAnimNode_OffsetRootBone::Initialize_AnyThread(const FAnimationInitializeContext& Context)
 {
-	FAnimNode_Base::Initialize_AnyThread(Context);
+	Super::Initialize_AnyThread(Context);
 	AnimInstanceProxy = Context.AnimInstanceProxy;
 	Source.Initialize(Context);
 	Reset(Context);
 }
 
+void FAnimNode_OffsetRootBone::CacheBones_AnyThread(const FAnimationCacheBonesContext& Context)
+{
+	Super::CacheBones_AnyThread(Context);
+	Source.CacheBones(Context);
+}
+
 void FAnimNode_OffsetRootBone::Update_AnyThread(const FAnimationUpdateContext& Context)
 {
-	FAnimNode_Base::Update_AnyThread(Context);
+	Super::Update_AnyThread(Context);
 	CachedDeltaTime = Context.GetDeltaTime();
 
 	// If we just became relevant and haven't been initialized yet, then reset.
diff --git a/Engine/Plugins/Animation/AnimationWarping/Source/Runtime/Public/BoneControllers/AnimNode_OffsetRootBone.h b/Engine/Plugins/Animation/AnimationWarping/Source/Runtime/Public/BoneControllers/AnimNode_OffsetRootBone.h
--- a/Engine/Plugins/Animation/AnimationWarping/Source/Runtime/Public/BoneControllers/AnimNode_OffsetRootBone.h
+++ b/Engine/Plugins/Animation/AnimationWarping/Source/Runtime/Public/BoneControllers/AnimNode_OffsetRootBone.h
@@ -107,6 +107,7 @@ public:
 	virtual void Update_AnyThread(const FAnimationUpdateContext& Context) override;
 	virtual void Evaluate_AnyThread(FPoseContext& Output) override;
 	virtual void Initialize_AnyThread(const FAnimationInitializeContext& Context) override;
+	virtual void CacheBones_AnyThread(const FAnimationCacheBonesContext& Context) override;
 	// End of FAnimNode_Base interface
 
 	// Folded property accesors
