Fix math in AnimNode_Steering for root motion codepath

diff --git a/Engine/Plugins/Animation/AnimationWarping/Source/Runtime/Private/BoneControllers/AnimNode_Steering.cpp b/Engine/Plugins/Animation/AnimationWarping/Source/Runtime/Private/BoneControllers/AnimNode_Steering.cpp
--- a/Engine/Plugins/Animation/AnimationWarping/Source/Runtime/Private/BoneControllers/AnimNode_Steering.cpp
+++ b/Engine/Plugins/Animation/AnimationWarping/Source/Runtime/Private/BoneControllers/AnimNode_Steering.cpp
@@ -17,7 +17,7 @@ void FAnimNode_Steering::UpdateInternal(const FAnimationUpdateContext& Context)
 
 	if (USkeletalMeshComponent* SkelMeshComponent = Context.AnimInstanceProxy->GetSkelMeshComponent())
 	{
-		RootBoneTransform = SkelMeshComponent->GetBoneTransform(0);
+		RootBoneTransform = SkelMeshComponent->GetBoneTransform(0); 
 	}
 	else
 	{
@@ -39,8 +39,8 @@ void FAnimNode_Steering::EvaluateSkeletalControl_AnyThread(FComponentSpacePoseCo
 
 		if (RootMotionProvider)
 		{
-			FTransform RootMotionTransformDelta = FTransform::Identity;
-			if (RootMotionProvider->ExtractRootMotion(Output.CustomAttributes, RootMotionTransformDelta))
+			FTransform ThisFrameRootMotionTransform = FTransform::Identity;
+			if (RootMotionProvider->ExtractRootMotion(Output.CustomAttributes, ThisFrameRootMotionTransform))
 			{
 				const float DeltaSeconds = Output.AnimInstanceProxy->GetDeltaSeconds();
 
@@ -48,53 +48,51 @@ void FAnimNode_Steering::EvaluateSkeletalControl_AnyThread(FComponentSpacePoseCo
 
 				UE_VLOG_ARROW(Output.AnimInstanceProxy->GetAnimInstanceObject(), "Steering", Display,
 					RootBoneTransform.GetLocation(),
-					RootBoneTransform.GetLocation()  + RootBoneRotation.GetForwardVector() * 100,
-					FColor::Red, TEXT(""));
+					RootBoneTransform.GetLocation()  + RootBoneRotation.GetRightVector() * 100,
+					FColor::Green, TEXT(""));
 					
 				UE_VLOG_ARROW(Output.AnimInstanceProxy->GetAnimInstanceObject(), "Steering", Display,
 					RootBoneTransform.GetLocation(),
-					RootBoneTransform.GetLocation()  + TargetOrientation.GetForwardVector() * 100,
-					FColor::Green, TEXT(""));
+					RootBoneTransform.GetLocation()  + TargetOrientation.GetRightVector() * 100,
+					FColor::Blue, TEXT(""));
 
-				FQuat Delta =  RootBoneRotation.Inverse() * TargetOrientation;
+				FQuat DeltaToTargetOrientation =  RootBoneRotation.Inverse() * TargetOrientation;
 
 
 				if (TargetTime > 0)
 				{
 					if (UAnimSequenceBase* AnimSequence = Cast<UAnimSequenceBase>(CurrentAnimAsset))
 					{
-						FTransform RootMotionDelta = AnimSequence->ExtractRootMotion(CurrentAnimAssetTime, TargetTime, true);
-						FQuat RootMotionRotation = RootMotionDelta.GetRotation();
-
-						FRotator RootMotionRot(RootMotionRotation);
+						FTransform PredictedRootMotionTransform = AnimSequence->ExtractRootMotion(CurrentAnimAssetTime, TargetTime, true);
+						FQuat PredictedRootMotionQuat = PredictedRootMotionTransform.GetRotation();
+						FRotator PredictedRootMotionRotation(PredictedRootMotionQuat);
 						
-						if (fabs(RootMotionRot.Yaw) > RootMotionThreshold)
+						if (fabs(PredictedRootMotionRotation.Yaw) > RootMotionThreshold)
 						{
 							UE_VLOG_ARROW(Output.AnimInstanceProxy->GetAnimInstanceObject(), "Steering", Display,
 								RootBoneTransform.GetLocation(),
-								RootBoneTransform.GetLocation()  + (RootMotionRotation * RootBoneRotation).GetForwardVector() * 100,
-								FColor::Blue, TEXT(""));
+								RootBoneTransform.GetLocation()  + (PredictedRootMotionQuat * RootBoneRotation).GetRightVector() * 100,
+								FColor::Orange, TEXT(""));
 
-							FRotator DeltaRot(Delta);
+							FRotator DeltaRot(DeltaToTargetOrientation);
 							
-							float Ratio =  (TargetTime / DeltaSeconds) * DeltaRot.Yaw / RootMotionRot.Yaw ;
-
-							FRotator RootMotionFrameRot(RootMotionTransformDelta.GetRotation());
+							float Ratio =  DeltaRot.Yaw / PredictedRootMotionRotation.Yaw ;
 
-							RootMotionFrameRot.Yaw *= Ratio;
+							FRotator ThisFrameRootMotionRotation(ThisFrameRootMotionTransform.GetRotation());
+							ThisFrameRootMotionRotation.Yaw *= Ratio;
 							
-							RootMotionTransformDelta.SetRotation(FQuat(RootMotionFrameRot));
-							RootMotionProvider->OverrideRootMotion(RootMotionTransformDelta, Output.CustomAttributes);
+							ThisFrameRootMotionTransform.SetRotation(FQuat(ThisFrameRootMotionRotation));
+							RootMotionProvider->OverrideRootMotion(ThisFrameRootMotionTransform, Output.CustomAttributes);
 							
 							return;
 						}
 					}
 					
-					Delta = FQuat::Slerp(FQuat::Identity, Delta,  DeltaSeconds/TargetTime);
+					DeltaToTargetOrientation = FQuat::Slerp(FQuat::Identity, DeltaToTargetOrientation,  DeltaSeconds/TargetTime);
 				}
 
-				RootMotionTransformDelta.SetRotation(FQuat::Slerp(RootMotionTransformDelta.GetRotation(), Delta, Alpha));
-				RootMotionProvider->OverrideRootMotion(RootMotionTransformDelta, Output.CustomAttributes);
+				ThisFrameRootMotionTransform.SetRotation(FQuat::Slerp(ThisFrameRootMotionTransform.GetRotation(), DeltaToTargetOrientation, Alpha));
+				RootMotionProvider->OverrideRootMotion(ThisFrameRootMotionTransform, Output.CustomAttributes);
 			}
 		}
 	}
