Fix for crash when missing a bone index in Foot Placement

diff --git a/Engine/Plugins/Animation/AnimationWarping/Source/Runtime/Private/BoneControllers/AnimNode_FootPlacement.cpp b/Engine/Plugins/Animation/AnimationWarping/Source/Runtime/Private/BoneControllers/AnimNode_FootPlacement.cpp
--- a/Engine/Plugins/Animation/AnimationWarping/Source/Runtime/Private/BoneControllers/AnimNode_FootPlacement.cpp
+++ b/Engine/Plugins/Animation/AnimationWarping/Source/Runtime/Private/BoneControllers/AnimNode_FootPlacement.cpp
@@ -1237,6 +1237,14 @@ void FAnimNode_FootPlacement::InitializeBoneReferences(const FBoneContainer& Req
 		LegData.Bones.IKIndex = LegDef.IKFootBone.GetCompactPoseIndex(RequiredBones);
 		LegData.Bones.FKIndex = LegDef.FKFootBone.GetCompactPoseIndex(RequiredBones);
 		LegData.Bones.BallIndex = LegDef.BallBone.GetCompactPoseIndex(RequiredBones);
+
+		if ((!LegData.Bones.IKIndex.IsValid()) ||
+			(!LegData.Bones.FKIndex.IsValid()) ||
+			(!LegData.Bones.BallIndex.IsValid()))
+		{
+			break;
+		}
+
 		UE::Anim::FootPlacement::FindChainLengthRootBoneIndex(
 			LegData.Bones.FKIndex, FMath::Max(LegDef.NumBonesInLimb, 1), RequiredBones,
 			LegData.Bones.HipIndex, LegData.Bones.LimbLength);
