Foot Placement: Fix for NaNs when delta time == 0.0 and stepping on vertical surfaces.

diff --git a/Engine/Plugins/Animation/AnimationWarping/Source/Runtime/Private/BoneControllers/AnimNode_FootPlacement.cpp b/Engine/Plugins/Animation/AnimationWarping/Source/Runtime/Private/BoneControllers/AnimNode_FootPlacement.cpp
--- a/Engine/Plugins/Animation/AnimationWarping/Source/Runtime/Private/BoneControllers/AnimNode_FootPlacement.cpp
+++ b/Engine/Plugins/Animation/AnimationWarping/Source/Runtime/Private/BoneControllers/AnimNode_FootPlacement.cpp
@@ -163,22 +163,15 @@ namespace UE::Anim::FootPlacement
 		const bool bHit = Context.World->SweepSingleByChannel(
 			HitResult, TraceStart, TraceEnd, FQuat::Identity, CollisionChannel, CollisionShape, QueryParams);
 
-		if (!bHit)
+		if (!bHit || !Context.GetMovementComponentIsWalkable(HitResult))
 		{
-			// If the hit fails, use the ground plane position and a default impact normal (negated trace direction)
+			// If the hit fails or isn't walkable, use the ground plane position and a default impact normal (negated trace direction)
 			OutImpactLocationWS = PointDirectionPlaneIntersection(StartPositionWS, TraceDirectionWS, FPlane(Context.GetMovementComponentFloorLocation(), -TraceDirectionWS));
 			OutImpactNormalWS = -TraceDirectionWS;
 			return false;
 		}
 
 		OutImpactLocationWS = HitResult.ImpactPoint;
-		if (!Context.GetMovementComponentIsWalkable(HitResult))
-		{
-			// If the surface hit isn't walkable, use the negated trace direction as the impact normal
-			OutImpactNormalWS = -TraceDirectionWS;
-			return false;
-		}
-
 		OutImpactNormalWS = HitResult.ImpactNormal;
 
 #if ENABLE_ANIM_DEBUG
@@ -1434,7 +1427,14 @@ void FAnimNode_FootPlacement::ProcessCharacterState(const UE::Anim::FootPlacemen
 		const FVector CapsuleMoveOffsetCS =
 			Context.OwningComponentToWorld.InverseTransformVectorNoScale(CapsuleMoveOffsetWS);
 
-		CharacterData.ComponentVelocityCS = CapsuleMoveOffsetCS / CachedDeltaTime;
+		if (CachedDeltaTime > 0.0f)
+		{
+			CharacterData.ComponentVelocityCS = CapsuleMoveOffsetCS / CachedDeltaTime;
+		}
+		else
+		{
+			CharacterData.ComponentVelocityCS = FVector::ZeroVector;
+		}
 	}
 }
 
