Anim Graph Steering  Node improvements

#changelist validated
#virtualized

diff --git a/Engine/Plugins/Animation/AnimationWarping/Source/Runtime/Private/BoneControllers/AnimNode_Steering.cpp b/Engine/Plugins/Animation/AnimationWarping/Source/Runtime/Private/BoneControllers/AnimNode_Steering.cpp
--- a/Engine/Plugins/Animation/AnimationWarping/Source/Runtime/Private/BoneControllers/AnimNode_Steering.cpp
+++ b/Engine/Plugins/Animation/AnimationWarping/Source/Runtime/Private/BoneControllers/AnimNode_Steering.cpp
@@ -6,6 +6,7 @@
 #include "Animation/AnimRootMotionProvider.h"
 #include "HAL/IConsoleManager.h"
 #include "Animation/AnimTrace.h"
+#include "Kismet/KismetMathLibrary.h"
 #include "Logging/LogVerbosity.h"
 #include "VisualLogger/VisualLogger.h"
 
@@ -28,71 +29,105 @@ void FAnimNode_Steering::UpdateInternal(const FAnimationUpdateContext& Context)
 void FAnimNode_Steering::Initialize_AnyThread(const FAnimationInitializeContext& Context)
 {
 	FAnimNode_SkeletalControlBase::Initialize_AnyThread(Context);
+	bResetFilter = true;
 }
 
 void FAnimNode_Steering::EvaluateSkeletalControl_AnyThread(FComponentSpacePoseContext& Output, TArray<FBoneTransform>& OutBoneTransforms)
 {
-	if (Alpha > 0.0f)
-	{
-		const UE::Anim::IAnimRootMotionProvider* RootMotionProvider = UE::Anim::IAnimRootMotionProvider::Get();
-		ensureMsgf(RootMotionProvider, TEXT("Steering expected a valid root motion delta provider interface."));
+	const float DeltaSeconds = Output.AnimInstanceProxy->GetDeltaSeconds();
 
-		if (RootMotionProvider)
+	if (DeltaSeconds > 0)
+	{
+		if (bEnableTargetSmoothing)
 		{
-			FTransform ThisFrameRootMotionTransform = FTransform::Identity;
-			if (RootMotionProvider->ExtractRootMotion(Output.CustomAttributes, ThisFrameRootMotionTransform))
+			if(bResetFilter)
 			{
-				const float DeltaSeconds = Output.AnimInstanceProxy->GetDeltaSeconds();
-
-				FQuat RootBoneRotation = RootBoneTransform.GetRotation();
-
-				UE_VLOG_ARROW(Output.AnimInstanceProxy->GetAnimInstanceObject(), "Steering", Display,
-					RootBoneTransform.GetLocation(),
-					RootBoneTransform.GetLocation()  + RootBoneRotation.GetRightVector() * 100,
-					FColor::Green, TEXT(""));
-					
-				UE_VLOG_ARROW(Output.AnimInstanceProxy->GetAnimInstanceObject(), "Steering", Display,
-					RootBoneTransform.GetLocation(),
-					RootBoneTransform.GetLocation()  + TargetOrientation.GetRightVector() * 100,
-					FColor::Blue, TEXT(""));
-
-				FQuat DeltaToTargetOrientation =  RootBoneRotation.Inverse() * TargetOrientation;
+				FilteredTarget = TargetOrientation;
+				TargetSmoothingState.Reset();
+				bResetFilter = false;
+			}
+			
+			FilteredTarget = UKismetMathLibrary::QuaternionSpringInterp(FilteredTarget, TargetOrientation, TargetSmoothingState,
+								SmoothTargetStiffness, SmoothTargetDamping, DeltaSeconds, 1, 0, true);
 
+			TargetOrientation = FilteredTarget;
+		}
+		
+		if (Alpha > 0.0f)
+		{
+			const UE::Anim::IAnimRootMotionProvider* RootMotionProvider = UE::Anim::IAnimRootMotionProvider::Get();
+			ensureMsgf(RootMotionProvider, TEXT("Steering expected a valid root motion delta provider interface."));
 
-				if (TargetTime > 0)
+			if (RootMotionProvider)
+			{
+				FTransform ThisFrameRootMotionTransform = FTransform::Identity;
+				if (RootMotionProvider->ExtractRootMotion(Output.CustomAttributes, ThisFrameRootMotionTransform))
 				{
-					if (UAnimSequenceBase* AnimSequence = Cast<UAnimSequenceBase>(CurrentAnimAsset))
+					if (ThisFrameRootMotionTransform.GetTranslation().Length() / DeltaSeconds > DisableSteeringBelowSpeed)
 					{
-						FTransform PredictedRootMotionTransform = AnimSequence->ExtractRootMotion(CurrentAnimAssetTime, TargetTime, true);
-						FQuat PredictedRootMotionQuat = PredictedRootMotionTransform.GetRotation();
-						FRotator PredictedRootMotionRotation(PredictedRootMotionQuat);
+						FQuat RootBoneRotation = RootBoneTransform.GetRotation();
+
+						UE_VLOG_ARROW(Output.AnimInstanceProxy->GetAnimInstanceObject(), "Steering", Display,
+							RootBoneTransform.GetLocation(),
+							RootBoneTransform.GetLocation()  + RootBoneRotation.GetRightVector() * 100,
+							FColor::Green, TEXT(""));
 						
-						if (fabs(PredictedRootMotionRotation.Yaw) > RootMotionThreshold)
-						{
-							UE_VLOG_ARROW(Output.AnimInstanceProxy->GetAnimInstanceObject(), "Steering", Display,
-								RootBoneTransform.GetLocation(),
-								RootBoneTransform.GetLocation()  + (PredictedRootMotionQuat * RootBoneRotation).GetRightVector() * 100,
-								FColor::Orange, TEXT(""));
+						UE_VLOG_ARROW(Output.AnimInstanceProxy->GetAnimInstanceObject(), "Steering", Display,
+							RootBoneTransform.GetLocation(),
+							RootBoneTransform.GetLocation()  + TargetOrientation.GetRightVector() * 100,
+							FColor::Blue, TEXT(""));
 
-							FRotator DeltaRot(DeltaToTargetOrientation);
-							
-							float Ratio =  DeltaRot.Yaw / PredictedRootMotionRotation.Yaw ;
+						FQuat DeltaToTargetOrientation =  RootBoneRotation.Inverse() * TargetOrientation;
 
-							FRotator ThisFrameRootMotionRotation(ThisFrameRootMotionTransform.GetRotation());
-							ThisFrameRootMotionRotation.Yaw *= Ratio;
-							
-							ThisFrameRootMotionTransform.SetRotation(FQuat(ThisFrameRootMotionRotation));
-							RootMotionProvider->OverrideRootMotion(ThisFrameRootMotionTransform, Output.CustomAttributes);
+						if (AnimatedTargetTime > 0)
+						{
+							if (UAnimSequenceBase* AnimSequence = Cast<UAnimSequenceBase>(CurrentAnimAsset))
+							{
+								FTransform PredictedRootMotionTransform = AnimSequence->ExtractRootMotion(CurrentAnimAssetTime, AnimatedTargetTime, true);
+								FQuat PredictedRootMotionQuat = PredictedRootMotionTransform.GetRotation();
+								float PredictedRootMotionYaw = FRotator(PredictedRootMotionQuat).Yaw;
 							
-							return;
+								if (fabs(PredictedRootMotionYaw) > RootMotionThreshold)
+								{
+									UE_VLOG_ARROW(Output.AnimInstanceProxy->GetAnimInstanceObject(), "Steering", Display,
+										RootBoneTransform.GetLocation(),
+										RootBoneTransform.GetLocation()  + (PredictedRootMotionQuat * RootBoneRotation).GetRightVector() * 100,
+										FColor::Orange, TEXT(""));
+
+									float YawToTargetOrientation = FRotator(DeltaToTargetOrientation).Yaw;
+
+									// pick the rotation direction that is the shortest path from the endpoint of the current animated rotation
+									if (PredictedRootMotionYaw - YawToTargetOrientation > 180)
+									{
+										YawToTargetOrientation += 360;
+									}
+									else if (YawToTargetOrientation - PredictedRootMotionYaw > 180)
+									{
+										YawToTargetOrientation -=360;
+									}
+
+									float Ratio =  YawToTargetOrientation / PredictedRootMotionYaw ;
+
+									FRotator ThisFrameRootMotionRotation(ThisFrameRootMotionTransform.GetRotation());
+									ThisFrameRootMotionRotation.Yaw *= Ratio;
+								
+									ThisFrameRootMotionTransform.SetRotation(FQuat(ThisFrameRootMotionRotation));
+									RootMotionProvider->OverrideRootMotion(ThisFrameRootMotionTransform, Output.CustomAttributes);
+								
+									return;
+								}
+							}
+						}
+						
+						if (ProceduralTargetTime > 0)
+						{	
+							DeltaToTargetOrientation = FQuat::Slerp(FQuat::Identity, DeltaToTargetOrientation,  DeltaSeconds/ProceduralTargetTime);
 						}
+
+						ThisFrameRootMotionTransform.SetRotation(FQuat::Slerp(ThisFrameRootMotionTransform.GetRotation(), DeltaToTargetOrientation, Alpha));
+						RootMotionProvider->OverrideRootMotion(ThisFrameRootMotionTransform, Output.CustomAttributes);
 					}
-					
-					DeltaToTargetOrientation = FQuat::Slerp(FQuat::Identity, DeltaToTargetOrientation,  DeltaSeconds/TargetTime);
 				}
-
-				ThisFrameRootMotionTransform.SetRotation(FQuat::Slerp(ThisFrameRootMotionTransform.GetRotation(), DeltaToTargetOrientation, Alpha));
-				RootMotionProvider->OverrideRootMotion(ThisFrameRootMotionTransform, Output.CustomAttributes);
 			}
 		}
 	}
diff --git a/Engine/Plugins/Animation/AnimationWarping/Source/Runtime/Public/BoneControllers/AnimNode_Steering.h b/Engine/Plugins/Animation/AnimationWarping/Source/Runtime/Public/BoneControllers/AnimNode_Steering.h
--- a/Engine/Plugins/Animation/AnimationWarping/Source/Runtime/Public/BoneControllers/AnimNode_Steering.h
+++ b/Engine/Plugins/Animation/AnimationWarping/Source/Runtime/Public/BoneControllers/AnimNode_Steering.h
@@ -4,6 +4,7 @@
 
 #include "BoneControllers/BoneControllerTypes.h"
 #include "BoneControllers/AnimNode_SkeletalControlBase.h"
+#include "Kismet/KismetMathLibrary.h"
 #include "AnimNode_Steering.generated.h"
 
 struct FAnimationInitializeContext;
@@ -19,24 +20,43 @@ struct ANIMATIONWARPINGRUNTIME_API FAnimNode_Steering : public FAnimNode_Skeleta
 	// The Orientation to steer towards
 	UPROPERTY(EditAnywhere, BlueprintReadWrite, Category=Evaluation, meta=(PinShownByDefault))
 	FQuat TargetOrientation = FQuat::Identity;
+
+	UPROPERTY(EditAnywhere, BlueprintReadWrite, Category=Filtering)
+	bool bEnableTargetSmoothing = false;
+	
+	UPROPERTY(EditAnywhere, BlueprintReadWrite, Category=Filtering)
+	float SmoothTargetStiffness = 300;
+	
+	UPROPERTY(EditAnywhere, BlueprintReadWrite, Category=Filtering)
+	float SmoothTargetDamping = 1;
+	
+	// The number of seconds in the future before we should reach the TargetOrientation when play animations with no root motion rotation
+	UPROPERTY(EditAnywhere, BlueprintReadWrite, Category=Evaluation, meta=(PinShownByDefault))
+	float ProceduralTargetTime = 0.2f;
 	
-	// The number of seconds in the future before we should reach the TargetOrientation
+	// Deprected old/unused parameter, to avoid breaking data
 	UPROPERTY(EditAnywhere, BlueprintReadWrite, Category=Evaluation, meta=(PinShownByDefault))
 	float TargetTime = 0.2f;
+	
+	// The number of seconds in the future before we should reach the TargetOrientation when playing animations with root motion rotation
+	UPROPERTY(EditAnywhere, BlueprintReadWrite, Category=Evaluation, meta=(PinShownByDefault))
+	float AnimatedTargetTime = 0.2f;
 
 	// If less than this number of degrees of rotation is found over TargetTime seconds in the CurrentAnimAsset, then rotation will be added linearly to reach the TargetOrientation
 	// Otherwise the root motion will be scaled to reach the TargetOrientation over TargetTime seconds
-	UPROPERTY(EditAnywhere, Category=Evaluation)
+	UPROPERTY(EditAnywhere, DisplayName=RootMotionAngleThreshold, Category=Evaluation)
 	float RootMotionThreshold = 1.0f;
-	
-	
+
+	// below this movement speed (based on the root motion in the animation) disable steering
+	UPROPERTY(EditAnywhere, Category=Evaluation)
+	float DisableSteeringBelowSpeed = 1.0f;
 	
 	// Animation Asset for incorporating root motion data. If CurrentAnimAsset is set, and the animation has root motion rotation within the TargetTime, then those rotations will be scaled to reach the TargetOrientation
-	UPROPERTY(EditAnywhere, BlueprintReadWrite, Category=Evaluation, meta=(PinShownByDefault))
+	UPROPERTY(EditAnywhere, Transient, BlueprintReadWrite, Category=Evaluation, meta=(PinShownByDefault))
 	TObjectPtr<UAnimationAsset> CurrentAnimAsset;
 	
 	// Current playback time in seconds of the CurrentAnimAsset
-	UPROPERTY(EditAnywhere, BlueprintReadWrite, Category=Evaluation, meta=(PinShownByDefault))
+	UPROPERTY(EditAnywhere, Transient, BlueprintReadWrite, Category=Evaluation, meta=(PinShownByDefault))
 	float CurrentAnimAssetTime = 0.f;
 
 	// FAnimNodeBase interface
@@ -52,4 +72,10 @@ struct ANIMATIONWARPINGRUNTIME_API FAnimNode_Steering : public FAnimNode_Skeleta
 private:
 
 	FTransform RootBoneTransform;
+
+	bool bResetFilter = true;
+
+	FQuat FilteredTarget = FQuat::Identity;
+	FQuaternionSpringState TargetSmoothingState;
+	
 };
