Foot Placement: Fix for issues w/ scaled mesh, and first frame floor interpolation
#11328
#10971
#rnx

diff --git a/Engine/Plugins/Animation/AnimationWarping/Source/Runtime/Private/BoneControllers/AnimNode_FootPlacement.cpp b/Engine/Plugins/Animation/AnimationWarping/Source/Runtime/Private/BoneControllers/AnimNode_FootPlacement.cpp
--- a/Engine/Plugins/Animation/AnimationWarping/Source/Runtime/Private/BoneControllers/AnimNode_FootPlacement.cpp
+++ b/Engine/Plugins/Animation/AnimationWarping/Source/Runtime/Private/BoneControllers/AnimNode_FootPlacement.cpp
@@ -94,7 +94,7 @@ namespace UE::Anim::FootPlacement
 		MovementComponent = CharacterOwner ? CharacterOwner->GetCharacterMovement() : nullptr;
 		OwningComponentToWorld = OwningComponent->GetComponentToWorld();
 
-		ApproachDirWS = OwningComponentToWorld.TransformVector(ApproachDirCS);
+		ApproachDirWS = OwningComponentToWorld.TransformVectorNoScale(ApproachDirCS);
 	
 		RootMotionTransformDelta = FTransform::Identity;
 		if (const UE::Anim::IAnimRootMotionProvider* RootMotionProvider = UE::Anim::IAnimRootMotionProvider::Get())
@@ -528,7 +528,7 @@ void FAnimNode_FootPlacement::UpdatePlantingPlaneInterpolation(
 		UE::Anim::FootPlacement::FindPlantPlane(Context, TraceSettings, FootTransformWS.GetLocation(), true, InOutPlantPlane, ImpactLocationWS);
 	}
 
-	if (InterpolationSettings.bEnableFloorInterpolation)
+	if (InterpolationSettings.bEnableFloorInterpolation && !bIsFirstUpdate)
 	{
 		FVector CurrPlaneIntersection = UE::Anim::FootPlacement::PointDirectionPlaneIntersection(
 			FootTransformWS.GetLocation(),
