Add direct access to OffsetRootBone current world transform via an anim node function and an animgraph message.
Use the AnimGraph Message in Steering to access the current root bone world transform.

diff --git a/Engine/Plugins/Animation/AnimationWarping/Source/Runtime/Private/AnimationWarpingLibrary.cpp b/Engine/Plugins/Animation/AnimationWarping/Source/Runtime/Private/AnimationWarpingLibrary.cpp
new file mode 100644
--- /dev/null
+++ b/Engine/Plugins/Animation/AnimationWarping/Source/Runtime/Private/AnimationWarpingLibrary.cpp
@@ -0,0 +1,16 @@
+// Copyright Epic Games, Inc. All Rights Reserved.
+
+#include "AnimationWarpingLibrary.h"
+#include "BoneControllers/AnimNode_OffsetRootBone.h"
+
+FTransform UAnimationWarpingLibrary::GetOffsetRootTransform(const FAnimNodeReference& Node)
+{
+	FTransform Transform(FTransform::Identity);
+
+	if (FAnimNode_OffsetRootBone* OffsetRoot = Node.GetAnimNodePtr<FAnimNode_OffsetRootBone>())
+	{
+		OffsetRoot->GetOffsetRootTransform(Transform);
+	}
+	
+	return Transform;
+}
diff --git a/Engine/Plugins/Animation/AnimationWarping/Source/Runtime/Private/BoneControllers/AnimNode_OffsetRootBone.cpp b/Engine/Plugins/Animation/AnimationWarping/Source/Runtime/Private/BoneControllers/AnimNode_OffsetRootBone.cpp
--- a/Engine/Plugins/Animation/AnimationWarping/Source/Runtime/Private/BoneControllers/AnimNode_OffsetRootBone.cpp
+++ b/Engine/Plugins/Animation/AnimationWarping/Source/Runtime/Private/BoneControllers/AnimNode_OffsetRootBone.cpp
@@ -17,6 +17,7 @@ TAutoConsoleVariable<int32> CVarAnimNodeOffsetRootBoneEnable(TEXT("a.AnimNode.Of
 TAutoConsoleVariable<int32> CVarAnimNodeOffsetRootBoneModifyBone(TEXT("a.AnimNode.OffsetRootBone.ModifyBone"), 1, TEXT("Toggle whether the transform is applied to the bone"));
 #endif
 
+IMPLEMENT_ANIMGRAPH_MESSAGE(UE::AnimationWarping::FRootOffsetProvider);
 
 namespace UE::Anim::OffsetRootBone
 {
@@ -72,6 +73,9 @@ void FAnimNode_OffsetRootBone::Update_AnyThread(const FAnimationUpdateContext& C
 		Reset(Context);
 	}
 	UpdateCounter.SynchronizeWith(Context.AnimInstanceProxy->GetUpdateCounter());
+	
+	UE::Anim::TScopedGraphMessage<UE::AnimationWarping::FRootOffsetProvider> ScopedMessage(Context, FTransform(SimulatedRotation, SimulatedTranslation));
+	
 	Source.Update(Context);
 }
 
@@ -266,9 +270,8 @@ void FAnimNode_OffsetRootBone::Evaluate_AnyThread(FPoseContext& Output)
 	{
 		static const TCHAR* LogName = TEXT("OffsetRootBone");
 		const float InnerCircleRadius = 40.0f;
-		const float CircleThickness = 3.0f;
+		const float CircleThickness = 1.5f;
 		const FVector CircleOffset(0,0,1);
-		const float ConeThickness = 0.3f;
 
 		const FTransform TargetBoneInitialTransformWorld = InputBoneTransform * ComponentTransform;
 		const FTransform TargetBoneTransformWorld = TargetBoneTransform * ComponentTransform;
@@ -279,18 +282,18 @@ void FAnimNode_OffsetRootBone::Evaluate_AnyThread(FPoseContext& Output)
 			const float OuterCircleRadius = GetMaxTranslationError() + InnerCircleRadius;
 			UE_VLOG_CIRCLE_THICK(AnimInstanceProxy->GetAnimInstanceObject(), TEXT("OffsetRootBone"), Display, ComponentTransform.GetLocation() + CircleOffset, FVector::UpVector, OuterCircleRadius, FColor::Red, CircleThickness, TEXT(""));
 		}
-
+		
 		UE_VLOG_CIRCLE_THICK(LogOwner, LogName, Display, ComponentTransform.GetLocation() + CircleOffset, FVector::UpVector, InnerCircleRadius, FColor::Blue, CircleThickness, TEXT(""));
-		UE_VLOG_SEGMENT_THICK(LogOwner, LogName, Display,
+		UE_VLOG_ARROW(LogOwner, LogName, Display,
 			ComponentTransform.GetLocation() + CircleOffset,
 			ComponentTransform.GetLocation() + InnerCircleRadius * ComponentTransform.GetRotation().GetRightVector() + CircleOffset,
-			FColor::Blue, CircleThickness, TEXT(""));
+			FColor::Blue, TEXT(""));
 		
 		UE_VLOG_CIRCLE_THICK(LogOwner, LogName, Display, TargetBoneTransformWorld.GetLocation() + CircleOffset, FVector::UpVector, InnerCircleRadius, FColor::Green, CircleThickness, TEXT(""));
-		UE_VLOG_SEGMENT_THICK(LogOwner, LogName, Display,
-		 	TargetBoneTransformWorld.GetLocation(),
+		UE_VLOG_ARROW(LogOwner, LogName, Display,
+		 	TargetBoneTransformWorld.GetLocation() + CircleOffset,
 		 	TargetBoneTransformWorld.GetLocation() + InnerCircleRadius * TargetBoneTransformWorld.GetRotation().GetRightVector() + CircleOffset,
-			FColor::Green, CircleThickness, TEXT(""));
+			FColor::Green, TEXT(""));
 	}
 #endif
 
diff --git a/Engine/Plugins/Animation/AnimationWarping/Source/Runtime/Private/BoneControllers/AnimNode_Steering.cpp b/Engine/Plugins/Animation/AnimationWarping/Source/Runtime/Private/BoneControllers/AnimNode_Steering.cpp
--- a/Engine/Plugins/Animation/AnimationWarping/Source/Runtime/Private/BoneControllers/AnimNode_Steering.cpp
+++ b/Engine/Plugins/Animation/AnimationWarping/Source/Runtime/Private/BoneControllers/AnimNode_Steering.cpp
@@ -6,6 +6,7 @@
 #include "Animation/AnimRootMotionProvider.h"
 #include "HAL/IConsoleManager.h"
 #include "Animation/AnimTrace.h"
+#include "BoneControllers/AnimNode_OffsetRootBone.h"
 #include "Kismet/KismetMathLibrary.h"
 #include "Logging/LogVerbosity.h"
 #include "VisualLogger/VisualLogger.h"
@@ -16,9 +17,9 @@ void FAnimNode_Steering::UpdateInternal(const FAnimationUpdateContext& Context)
 {
 	FAnimNode_SkeletalControlBase::UpdateInternal(Context);
 
-	if (USkeletalMeshComponent* SkelMeshComponent = Context.AnimInstanceProxy->GetSkelMeshComponent())
+	if (UE::AnimationWarping::FRootOffsetProvider* RootOffsetProvider = Context.GetMessage<UE::AnimationWarping::FRootOffsetProvider>())
 	{
-		RootBoneTransform = SkelMeshComponent->GetBoneTransform(0); 
+		RootBoneTransform = RootOffsetProvider->GetRootTransform();
 	}
 	else
 	{
@@ -32,6 +33,7 @@ void FAnimNode_Steering::Initialize_AnyThread(const FAnimationInitializeContext&
 	bResetFilter = true;
 }
 
+
 void FAnimNode_Steering::EvaluateSkeletalControl_AnyThread(FComponentSpacePoseContext& Output, TArray<FBoneTransform>& OutBoneTransforms)
 {
 	const float DeltaSeconds = Output.AnimInstanceProxy->GetDeltaSeconds();
@@ -69,7 +71,7 @@ void FAnimNode_Steering::EvaluateSkeletalControl_AnyThread(FComponentSpacePoseCo
 
 						UE_VLOG_ARROW(Output.AnimInstanceProxy->GetAnimInstanceObject(), "Steering", Display,
 							RootBoneTransform.GetLocation(),
-							RootBoneTransform.GetLocation()  + RootBoneRotation.GetRightVector() * 100,
+							RootBoneTransform.GetLocation()  + RootBoneRotation.GetRightVector() * 90,
 							FColor::Green, TEXT(""));
 						
 						UE_VLOG_ARROW(Output.AnimInstanceProxy->GetAnimInstanceObject(), "Steering", Display,
@@ -119,7 +121,7 @@ void FAnimNode_Steering::EvaluateSkeletalControl_AnyThread(FComponentSpacePoseCo
 							}
 						}
 						
-						if (ProceduralTargetTime > 0)
+						if (ProceduralTargetTime > DeltaSeconds)
 						{	
 							DeltaToTargetOrientation = FQuat::Slerp(FQuat::Identity, DeltaToTargetOrientation,  DeltaSeconds/ProceduralTargetTime);
 						}
@@ -131,4 +133,4 @@ void FAnimNode_Steering::EvaluateSkeletalControl_AnyThread(FComponentSpacePoseCo
 			}
 		}
 	}
-}
\ No newline at end of file
+}
diff --git a/Engine/Plugins/Animation/AnimationWarping/Source/Runtime/Public/AnimationWarpingLibrary.h b/Engine/Plugins/Animation/AnimationWarping/Source/Runtime/Public/AnimationWarpingLibrary.h
new file mode 100644
--- /dev/null
+++ b/Engine/Plugins/Animation/AnimationWarping/Source/Runtime/Public/AnimationWarpingLibrary.h
@@ -0,0 +1,23 @@
+// Copyright Epic Games, Inc. All Rights Reserved.
+
+#pragma once
+
+#include "Animation/AnimExecutionContext.h"
+#include "Animation/AnimNodeReference.h"
+#include "Kismet/BlueprintFunctionLibrary.h"
+#include "AnimationWarpingLibrary.generated.h"
+
+struct FAnimNode_OffsetRootBone;
+
+// Exposes operations related to Animation Warping
+UCLASS(Experimental)
+class ANIMATIONWARPINGRUNTIME_API UAnimationWarpingLibrary : public UBlueprintFunctionLibrary
+{
+	GENERATED_BODY()
+
+public:
+	/** Get the current world space transform from the offset root bone animgraph node */
+    UFUNCTION(BlueprintPure, Category = "Animation|BlendStack", meta = (BlueprintThreadSafe))
+    static FTransform GetOffsetRootTransform(const FAnimNodeReference& Node);
+	
+};
\ No newline at end of file
diff --git a/Engine/Plugins/Animation/AnimationWarping/Source/Runtime/Public/BoneControllers/AnimNode_OffsetRootBone.h b/Engine/Plugins/Animation/AnimationWarping/Source/Runtime/Public/BoneControllers/AnimNode_OffsetRootBone.h
--- a/Engine/Plugins/Animation/AnimationWarping/Source/Runtime/Public/BoneControllers/AnimNode_OffsetRootBone.h
+++ b/Engine/Plugins/Animation/AnimationWarping/Source/Runtime/Public/BoneControllers/AnimNode_OffsetRootBone.h
@@ -10,6 +10,26 @@ struct FAnimationInitializeContext;
 struct FComponentSpacePoseContext;
 struct FNodeDebugData;
 
+namespace UE::AnimationWarping
+{
+class FRootOffsetProvider : public UE::Anim::IGraphMessage
+{
+	DECLARE_ANIMGRAPH_MESSAGE(FRootOffsetProvider);
+
+public:
+
+	FRootOffsetProvider(const FTransform& InRootTransform)
+		: RootTransform(InRootTransform)
+	{
+	}
+	
+	const FTransform& GetRootTransform() const { return RootTransform; } 
+
+private:
+	FTransform RootTransform;
+};
+}
+
 UENUM(BlueprintType)
 enum class EOffsetRootBoneMode : uint8
 {
@@ -125,6 +145,14 @@ public:
 	float GetTranslationSpeedRatio() const;
 	float GetRotationSpeedRatio() const;
 
+
+	// get the current simulated root transform
+	void GetOffsetRootTransform(FTransform& OutTransform)
+	{
+		OutTransform.SetRotation(SimulatedRotation);
+		OutTransform.SetTranslation(SimulatedTranslation);
+	}
+
 private:
 
 	void Reset(const FAnimationBaseContext& Context);
