Foot placement changes.  Add visual logging draws, and add curves for disabling each foot + pelvis

diff --git a/Engine/Plugins/Animation/AnimationWarping/Source/Runtime/Private/BoneControllers/AnimNode_FootPlacement.cpp b/Engine/Plugins/Animation/AnimationWarping/Source/Runtime/Private/BoneControllers/AnimNode_FootPlacement.cpp
--- a/Engine/Plugins/Animation/AnimationWarping/Source/Runtime/Private/BoneControllers/AnimNode_FootPlacement.cpp
+++ b/Engine/Plugins/Animation/AnimationWarping/Source/Runtime/Private/BoneControllers/AnimNode_FootPlacement.cpp
@@ -8,12 +8,13 @@
 #include "GameFramework/Character.h"
 #include "GameFramework/CharacterMovementComponent.h"
 #include "Animation/AnimRootMotionProvider.h"
+#include "VisualLogger/VisualLogger.h"
 
 #include UE_INLINE_GENERATED_CPP_BY_NAME(AnimNode_FootPlacement)
 
 DECLARE_CYCLE_STAT(TEXT("Foot Placement Eval"), STAT_FootPlacement_Eval, STATGROUP_Anim);
 
-#if ENABLE_ANIM_DEBUG
+#if ENABLE_FOOTPLACEMENT_DEBUG
 static TAutoConsoleVariable<bool> CVarAnimNodeFootPlacementEnable(TEXT("a.AnimNode.FootPlacement.Enable"), true, TEXT("Enable/Disable Foot Placement"));
 static TAutoConsoleVariable<bool> CVarAnimNodeFootPlacementEnableLock(TEXT("a.AnimNode.FootPlacement.Enable.Lock"), true, TEXT("Enable/Disable Foot Locking"));
 static TAutoConsoleVariable<bool> CVarAnimNodeFootPlacementDebug(TEXT("a.AnimNode.FootPlacement.Debug"), false, TEXT("Turn on visualization debugging for Foot Placement"));
@@ -174,7 +175,7 @@ namespace UE::Anim::FootPlacement
 		OutImpactLocationWS = HitResult.ImpactPoint;
 		OutImpactNormalWS = HitResult.ImpactNormal;
 
-#if ENABLE_ANIM_DEBUG
+#if ENABLE_FOOTPLACEMENT_DEBUG
 		if (CVarAnimNodeFootPlacementDebugTraces.GetValueOnAnyThread())
 		{
 			Context.CSPContext.AnimInstanceProxy->AnimDrawDebugPoint(
@@ -682,7 +683,7 @@ void FAnimNode_FootPlacement::ResetRuntimeData()
 		LegData.Interpolation = UE::Anim::FootPlacement::FLegRuntimeData::FInterpolationData();
 	}
 
-#if ENABLE_ANIM_DEBUG
+#if ENABLE_FOOTPLACEMENT_DEBUG
 	DebugData.Init(LegDefinitions.Num());
 #endif
 
@@ -693,7 +694,7 @@ bool FAnimNode_FootPlacement::WantsToPlant(
 	const UE::Anim::FootPlacement::FEvaluationContext& Context,
 	const UE::Anim::FootPlacement::FLegRuntimeData::FInputPoseData& LegInputPose) const
 {
-#if ENABLE_ANIM_DEBUG
+#if ENABLE_FOOTPLACEMENT_DEBUG
 	if (!CVarAnimNodeFootPlacementEnableLock.GetValueOnAnyThread())
 	{
 		return false;
@@ -799,7 +800,7 @@ UE::Anim::FootPlacement::FPlantResult FAnimNode_FootPlacement::FinalizeFootAlign
 					const FVector CorrectedFootToToe = CorrectedBallLocationCS - CorrectedFootLocationCS;
 					FQuat DeltaSlopeRotation = FQuat::FindBetweenVectors(InitialFootToToe, CorrectedFootToToe);
 
-#if ENABLE_ANIM_DEBUG
+#if ENABLE_FOOTPLACEMENT_DEBUG
 					FRotator DeltaSlopeRotator = DeltaSlopeRotation.Rotator();
 #endif
 
@@ -828,11 +829,11 @@ UE::Anim::FootPlacement::FPlantResult FAnimNode_FootPlacement::FinalizeFootAlign
 				LegData.Plant.bCanReachTarget = true;
 			}
 
-#if ENABLE_ANIM_DEBUG
+#if ENABLE_FOOTPLACEMENT_DEBUG
 			DebugData.LegsInfo[LegData.Idx].HyperExtensionAmount = HyperExtensionAmount;
 			DebugData.LegsInfo[LegData.Idx].RollAmount = HyperExtensionAmount - HyperExtensionRemaining;
 			DebugData.LegsInfo[LegData.Idx].PullAmount = FMath::Max(0.0f, HyperExtensionRemaining);
-#endif // (ENABLE_ANIM_DEBUG)
+#endif // (ENABLE_FOOTPLACEMENT_DEBUG)
 		}
 	}
 
@@ -887,7 +888,117 @@ UE::Anim::FootPlacement::FPlantResult FAnimNode_FootPlacement::FinalizeFootAlign
 	return Result;
 }
 
-#if ENABLE_ANIM_DEBUG
+#if ENABLE_FOOTPLACEMENT_DEBUG
+void FAnimNode_FootPlacement::DrawVLog(
+	const UE::Anim::FootPlacement::FEvaluationContext& Context,
+	const UE::Anim::FootPlacement::FLegRuntimeData& LegData,
+	const UE::Anim::FootPlacement::FPlantResult& PlantResult) const
+{
+	using namespace UE::Anim::FootPlacement;
+
+	const FColor FKColor = FColor::Blue;
+	const FColor PlantedColor = FColor::Red;
+	const FColor UnplantedColor = FColor::Green;
+	const FColor ReplantedColor = FColor::Orange;
+
+	FColor CurrentPlantColor;
+	switch (LegData.Plant.PlantType)
+	{
+	case UE::Anim::FootPlacement::EPlantType::Planted: CurrentPlantColor = PlantedColor; break;
+	case UE::Anim::FootPlacement::EPlantType::Unplanted: CurrentPlantColor = UnplantedColor; break;
+	case UE::Anim::FootPlacement::EPlantType::Replanted: CurrentPlantColor = ReplantedColor; break;
+	default: check(false); break; //not implemented
+	}
+
+	const FTransform FKBoneTransformWS =
+		LegData.InputPose.FootToGround *
+		LegData.InputPose.FootTransformCS *
+		Context.OwningComponentToWorld;
+
+	const FTransform IKBoneTransformWS =
+		LegData.InputPose.FootToGround *
+		LegData.AlignedFootTransformWS;
+
+
+	const FVector FKBoneLocationProjectedWS = UE::Anim::FootPlacement::PointDirectionPlaneIntersection(
+		FKBoneTransformWS.GetLocation(),
+		Context.ApproachDirWS,
+		LegData.Plant.PlantPlaneWS);
+	
+	UObject* LogOwner = Context.CSPContext.AnimInstanceProxy->GetAnimInstanceObject();
+	FName LogCategory = "FootPlacement";
+
+	UE_VLOG_SPHERE(LogOwner, LogCategory, Display, FKBoneTransformWS.GetLocation(), 0, FKColor, TEXT(""));
+	
+	UE_VLOG_SPHERE(LogOwner, LogCategory, Display, FKBoneLocationProjectedWS, 0, FKColor, TEXT(""));
+
+	UE_VLOG_SPHERE(LogOwner, LogCategory, Display, IKBoneTransformWS.GetLocation(), 0, CurrentPlantColor, TEXT(""));
+
+	const FVector IKBoneLocationProjectedWS = UE::Anim::FootPlacement::PointDirectionPlaneIntersection(
+		IKBoneTransformWS.GetLocation(),
+		Context.ApproachDirWS,
+		LegData.Plant.PlantPlaneWS);
+
+	UE_VLOG_SPHERE(LogOwner, LogCategory, Display, IKBoneLocationProjectedWS, 0, CurrentPlantColor, TEXT(""));
+	UE_VLOG_SEGMENT_THICK(LogOwner, LogCategory, Display, IKBoneTransformWS.GetLocation(), IKBoneLocationProjectedWS, CurrentPlantColor, 2, TEXT(""));
+
+	const float UnplantRadius = PlantSettings.UnplantRadius;
+	const FVector PlantCenter = UE::Anim::FootPlacement::PointDirectionPlaneIntersection(
+		IKBoneTransformWS.GetLocation(),
+		Context.ApproachDirWS,
+		LegData.Plant.PlantPlaneWS);
+
+	UE_VLOG_CIRCLE(LogOwner, LogCategory, Display, PlantCenter, LegData.Plant.PlantPlaneWS.GetNormal(), UnplantRadius, PlantedColor, TEXT(""));
+	
+	if (PlantSettings.ReplantRadiusRatio < 1.0f)
+	{
+		const float ReplantRadius =
+			PlantSettings.UnplantRadius *
+			PlantSettings.ReplantRadiusRatio;
+		
+		UE_VLOG_CIRCLE(LogOwner, LogCategory, Display, PlantCenter, LegData.Plant.PlantPlaneWS.GetNormal(), ReplantRadius, ReplantedColor, TEXT(""));
+	}
+
+	// FString InputPoseMessage = FString::Printf(
+	// 	TEXT("%s\n\t - InputPose [ AlignmentAlpha = %.2f, Speed = %.2f, DistanceToPlant = %.2f]"), 
+	// 		*LegDefinitions[LegData.Idx].FKFootBone.BoneName.ToString(),
+	// 		LegData.InputPose.AlignmentAlpha,
+	// 		LegData.InputPose.Speed,
+	// 		LegData.InputPose.DistanceToPlant );
+	// Context.CSPContext.AnimInstanceProxy->AnimDrawDebugOnScreenMessage(InputPoseMessage, FColor::White);
+	//
+	// FString ExtensionMessage = FString::Printf(
+	// 	TEXT("\t - HyperExtension[ Amount = %.2f, Roll = %.2f, Pull %.2f]"),
+	// 		DebugData.LegsInfo[LegData.Idx].HyperExtensionAmount,
+	// 		DebugData.LegsInfo[LegData.Idx].RollAmount,
+	// 		DebugData.LegsInfo[LegData.Idx].PullAmount);
+	// Context.CSPContext.AnimInstanceProxy->AnimDrawDebugOnScreenMessage(ExtensionMessage,
+	// 	(DebugData.LegsInfo[LegData.Idx].HyperExtensionAmount <= 0.0f) ? FColor::Green : FColor::Red);
+
+	if (PlantSettings.SeparatingDistance > 0.0f)
+	{
+		// FString SeparationPlaneMessage = FString::Printf(
+		// 	TEXT("\t - Distance To Separating Plane = %.2f"),
+		// 	DebugData.LegsInfo[LegData.Idx].DistanceToSeparatingPlane);
+		// Context.CSPContext.AnimInstanceProxy->AnimDrawDebugOnScreenMessage(SeparationPlaneMessage,
+		// 	(DebugData.LegsInfo[LegData.Idx].DistanceToSeparatingPlane < 0.0f) ? FColor::Red : FColor::Green);
+
+		TRACE_ANIM_NODE_VALUE(Context.CSPContext, TEXT("DistanceToSeparatingPlane"), DebugData.LegsInfo[LegData.Idx].DistanceToSeparatingPlane);
+	}
+
+	TRACE_ANIM_NODE_VALUE(Context.CSPContext, TStringBuilder<256>().Append("HyperExtension - ").Append(FString::FromInt(LegData.Idx)).ToString(), DebugData.LegsInfo[LegData.Idx].HyperExtensionAmount);
+	TRACE_ANIM_NODE_VALUE(Context.CSPContext, TStringBuilder<256>().Append("Roll - ").Append(FString::FromInt(LegData.Idx)).ToString(), DebugData.LegsInfo[LegData.Idx].RollAmount);
+	TRACE_ANIM_NODE_VALUE(Context.CSPContext, TStringBuilder<256>().Append("Pull - ").Append(FString::FromInt(LegData.Idx)).ToString(), DebugData.LegsInfo[LegData.Idx].PullAmount);
+	TRACE_ANIM_NODE_VALUE(Context.CSPContext, TStringBuilder<256>().Append("AlignmentAlpha - ").Append(FString::FromInt(LegData.Idx)).ToString(), LegData.InputPose.AlignmentAlpha);
+	
+	TRACE_ANIM_NODE_VALUE(Context.CSPContext, TStringBuilder<256>().Append("FootSpeed - ").Append(FString::FromInt(LegData.Idx)).ToString(), LegData.InputPose.Speed);
+	TRACE_ANIM_NODE_VALUE(Context.CSPContext, TStringBuilder<256>().Append("DistanceToPlant - ").Append(FString::FromInt(LegData.Idx)).ToString(), LegData.InputPose.DistanceToPlant);
+
+}
+#endif
+
+
+#if ENABLE_FOOTPLACEMENT_DEBUG
 void FAnimNode_FootPlacement::DrawDebug(
 	const UE::Anim::FootPlacement::FEvaluationContext& Context,
 	const UE::Anim::FootPlacement::FLegRuntimeData& LegData,
@@ -990,13 +1101,7 @@ void FAnimNode_FootPlacement::DrawDebug(
 		Context.CSPContext.AnimInstanceProxy->AnimDrawDebugOnScreenMessage(SeparationPlaneMessage,
 			(DebugData.LegsInfo[LegData.Idx].DistanceToSeparatingPlane < 0.0f) ? FColor::Red : FColor::Green);
 
-		TRACE_ANIM_NODE_VALUE(Context.CSPContext, TEXT("DistanceToSeparatingPlane"), DebugData.LegsInfo[LegData.Idx].DistanceToSeparatingPlane);
 	}
-
-	TRACE_ANIM_NODE_VALUE(Context.CSPContext, TEXT("HyperExtension - Amount"), DebugData.LegsInfo[LegData.Idx].HyperExtensionAmount);
-	TRACE_ANIM_NODE_VALUE(Context.CSPContext, TEXT("HyperExtension - Roll"), DebugData.LegsInfo[LegData.Idx].RollAmount);
-	TRACE_ANIM_NODE_VALUE(Context.CSPContext, TEXT("HyperExtension - Pull"), DebugData.LegsInfo[LegData.Idx].PullAmount);
-	TRACE_ANIM_NODE_VALUE(Context.CSPContext, TEXT("InputPose - AlignmentAlpha"), LegData.InputPose.AlignmentAlpha);
 }
 #endif
 
@@ -1053,7 +1158,7 @@ void FAnimNode_FootPlacement::EvaluateSkeletalControl_AnyThread(FComponentSpaceP
 		}
 	}
 	
-#if ENABLE_ANIM_DEBUG
+#if ENABLE_FOOTPLACEMENT_DEBUG
 	UE::Anim::FootPlacement::FDebugData LastDebugData = DebugData;
 #endif
 
@@ -1091,7 +1196,11 @@ void FAnimNode_FootPlacement::EvaluateSkeletalControl_AnyThread(FComponentSpaceP
 
 	// Based on the ground alignment, search for the best Pelvis transform
 	FTransform PelvisTransformCS = SolvePelvis(FootPlacementContext);
-#if ENABLE_ANIM_DEBUG
+
+	float DisablePelvisCurveValue = Output.Curve.Get(PelvisSettings.DisablePelvisCurveName);
+	PelvisTransformCS.BlendWith(PelvisData.InputPose.FKTransformCS, DisablePelvisCurveValue);
+
+#if ENABLE_FOOTPLACEMENT_DEBUG
 	const FTransform PelvisTargetTransformCS = PelvisTransformCS;
 	if (CVarAnimNodeFootPlacementDebug.GetValueOnAnyThread())
 	{
@@ -1118,7 +1227,14 @@ void FAnimNode_FootPlacement::EvaluateSkeletalControl_AnyThread(FComponentSpaceP
 		//OutBoneTransforms.Add(PlantResult.BallTransformCS);
 		//OutBoneTransforms.Add(PlantResult.HipTransformCS);
 
-#if ENABLE_ANIM_DEBUG
+#if ENABLE_FOOTPLACEMENT_DEBUG
+		if (FVisualLogger::IsRecording())
+		{
+			DrawVLog(FootPlacementContext, LegData, PlantResult);
+		}
+#endif
+
+#if ENABLE_FOOTPLACEMENT_DEBUG
 		if (CVarAnimNodeFootPlacementDebug.GetValueOnAnyThread())
 		{
 			DrawDebug(FootPlacementContext, LegData, PlantResult);
@@ -1136,7 +1252,7 @@ void FAnimNode_FootPlacement::EvaluateSkeletalControl_AnyThread(FComponentSpaceP
 
 	CachedDeltaTime = 0.0f;
 
-#if ENABLE_ANIM_DEBUG
+#if ENABLE_FOOTPLACEMENT_DEBUG
 	{
 		FAnimInstanceProxy* AnimInstanceProxy = Output.AnimInstanceProxy;
 		const FTransform ComponentTransform =
@@ -1180,6 +1296,21 @@ void FAnimNode_FootPlacement::EvaluateSkeletalControl_AnyThread(FComponentSpaceP
 			AnimInstanceProxy->AnimDrawDebugPoint(
 				PelvisTargetTransformWS.GetLocation(), 10.0f, FColor::Purple, false, -1.0f, SDPG_Foreground);
 		}
+		
+#if ENABLE_FOOTPLACEMENT_DEBUG
+		if (FVisualLogger::IsRecording())
+		{
+			// pelvis debugging
+			const FTransform PelvisTransformWS = PelvisTransformCS * ComponentTransform;
+			const FTransform BasePelvisTransformWS = PelvisData.InputPose.FKTransformCS * ComponentTransform;
+			const FTransform PelvisTargetTransformWS = PelvisTargetTransformCS * ComponentTransform;
+
+			UObject* AnimInstance = AnimInstanceProxy->GetAnimInstanceObject();
+			UE_VLOG_SPHERE(AnimInstance, "FootPlacement", Display, PelvisTransformWS.GetTranslation(), 0, FColor::Green, TEXT(""));
+			UE_VLOG_SPHERE(AnimInstance, "FootPlacement", Display, BasePelvisTransformWS.GetTranslation(), 0, FColor::Blue, TEXT(""));
+			UE_VLOG_SPHERE(AnimInstance, "FootPlacement", Display, PelvisTargetTransformWS.GetTranslation(), 0, FColor::Purple, TEXT(""));
+		}
+#endif
 	}
 #endif
 
@@ -1190,7 +1321,7 @@ void FAnimNode_FootPlacement::EvaluateSkeletalControl_AnyThread(FComponentSpaceP
 
 bool FAnimNode_FootPlacement::IsValidToEvaluate(const USkeleton* Skeleton, const FBoneContainer& RequiredBones)
 {
-#if ENABLE_ANIM_DEBUG
+#if ENABLE_FOOTPLACEMENT_DEBUG
 	if (!CVarAnimNodeFootPlacementEnable.GetValueOnAnyThread())
 	{
 		return false;
@@ -1245,7 +1376,7 @@ void FAnimNode_FootPlacement::InitializeBoneReferences(const FBoneContainer& Req
 		const FTransform BallTransformLS = RequiredBones.GetRefPoseTransform(LegData.Bones.BallIndex);
 		LegData.Bones.FootLength = BallTransformLS.GetLocation().Size();
 
-#if ENABLE_ANIM_DEBUG
+#if ENABLE_FOOTPLACEMENT_DEBUG
 		// TODO: This wont work for animations authored for different slopes or stairs. Figure this out later
 		const FVector RefPoseGroundNormalCS = FVector::UpVector;
 		const FTransform BallRefTransformCS = FAnimationRuntime::GetComponentSpaceRefPose(
@@ -1571,7 +1702,7 @@ void FAnimNode_FootPlacement::ProcessFootAlignment(
 			FootUnalignedTransformCS.SetLocation(FootUnalignedLocationCS);
 		}
 
-#if ENABLE_ANIM_DEBUG
+#if ENABLE_FOOTPLACEMENT_DEBUG
 		if (CVarAnimNodeFootPlacementDebug.GetValueOnAnyThread())
 		{
 			const FVector PlaneNormalWS = Context.OwningComponentToWorld.TransformVector(PlaneNormal);;
@@ -1583,6 +1714,16 @@ void FAnimNode_FootPlacement::ProcessFootAlignment(
 			DebugData.LegsInfo[LegData.Idx].DistanceToSeparatingPlane = DistanceToSeparatingPlane;
 		}
 #endif
+		
+#if ENABLE_FOOTPLACEMENT_DEBUG
+		if (FVisualLogger::IsRecording())
+		{
+			const FVector PlaneNormalWS = Context.OwningComponentToWorld.TransformVector(PlaneNormal);;
+			const FVector PlaneCenterWS = Context.OwningComponentToWorld.TransformPosition(PlaneCenter);;
+			
+			UE_VLOG_CIRCLE_THICK(Context.CSPContext.AnimInstanceProxy->GetAnimInstanceObject(), "FootPlacement", Display, PlaneCenterWS, PlaneNormalWS, 25, FColor::Red, 1, TEXT(""));
+		}
+#endif
 	}
 
 	FTransform BlendedUnalignedTransformCS;
@@ -1622,7 +1763,11 @@ void FAnimNode_FootPlacement::ProcessFootAlignment(
 			LegData.AlignedFootTransformCS.GetRotation(),
 			InputPose.AlignmentAlpha));
 
+	
 	LegData.AlignedFootTransformCS = BlendedPlantTransformCS;
+	
+	float DisableLegCurveValue = Context.CSPContext.Curve.Get(PelvisSettings.DisablePelvisCurveName);
+	LegData.AlignedFootTransformCS.BlendWith(LegData.InputPose.FootTransformCS, DisableLegCurveValue);
 }
 
 FVector FAnimNode_FootPlacement::GetApproachDirWS(const FAnimationBaseContext& Context) const
diff --git a/Engine/Plugins/Animation/AnimationWarping/Source/Runtime/Public/BoneControllers/AnimNode_FootPlacement.h b/Engine/Plugins/Animation/AnimationWarping/Source/Runtime/Public/BoneControllers/AnimNode_FootPlacement.h
--- a/Engine/Plugins/Animation/AnimationWarping/Source/Runtime/Public/BoneControllers/AnimNode_FootPlacement.h
+++ b/Engine/Plugins/Animation/AnimationWarping/Source/Runtime/Public/BoneControllers/AnimNode_FootPlacement.h
@@ -6,6 +6,7 @@
 #include "BoneControllers/AnimNode_SkeletalControlBase.h"
 #include "Kismet/KismetMathLibrary.h"
 #include "Engine/EngineTypes.h"
+#include "VisualLogger/VisualLogger.h"
 
 #include "AnimNode_FootPlacement.generated.h"
 
@@ -13,6 +14,12 @@ struct FAnimationInitializeContext;
 struct FComponentSpacePoseContext;
 struct FNodeDebugData;
 
+#if ENABLE_ANIM_DEBUG && ENABLE_VISUAL_LOG
+#define ENABLE_FOOTPLACEMENT_DEBUG 1
+#else
+#define  ENABLE_FOOTPLACEMENT_DEBUG 0
+#endif
+
 namespace UE::Anim::FootPlacement
 {
 	enum class EPlantType
@@ -372,6 +379,9 @@ public:
 
 	UPROPERTY(EditAnywhere, Category = "Pelvis Settings")
 	bool bEnableInterpolation = true;
+	
+	UPROPERTY(EditAnywhere, Category = "Pelvis Settings")
+	FName DisablePelvisCurveName = NAME_None;
 };
 
 USTRUCT()
@@ -403,6 +413,10 @@ public:
 	// This allows you to disable locking precisely, instead of relying on the procedural mechanism based on springs and foot analysis
 	UPROPERTY(EditAnywhere, Category = "Settings")
 	FName DisableLockCurveName = NAME_None;
+	
+	// Curve to disable the effect of footplacement on this leg
+	UPROPERTY(EditAnywhere, Category = "Settings")
+	FName DisableLegCurveName = NAME_None;
 
 public:
 
@@ -650,12 +664,19 @@ private:
 
 	void ResetRuntimeData();
 
-#if ENABLE_ANIM_DEBUG
+	
+#if ENABLE_FOOTPLACEMENT_DEBUG
 	UE::Anim::FootPlacement::FDebugData DebugData;
+
 	void DrawDebug(
 		const UE::Anim::FootPlacement::FEvaluationContext& Context,
 		const UE::Anim::FootPlacement::FLegRuntimeData& LegData,
 		const UE::Anim::FootPlacement::FPlantResult& PlantResult) const;
+
+	void DrawVLog(
+		const UE::Anim::FootPlacement::FEvaluationContext& Context,
+		const UE::Anim::FootPlacement::FLegRuntimeData& LegData,
+    	const UE::Anim::FootPlacement::FPlantResult& PlantResult) const;
 #endif
 
 	bool bIsFirstUpdate = false;
