Build fix

diff --git a/Engine/Plugins/Animation/AnimationWarping/Source/Runtime/Private/BoneControllers/AnimNode_OffsetRootBone.cpp b/Engine/Plugins/Animation/AnimationWarping/Source/Runtime/Private/BoneControllers/AnimNode_OffsetRootBone.cpp
--- a/Engine/Plugins/Animation/AnimationWarping/Source/Runtime/Private/BoneControllers/AnimNode_OffsetRootBone.cpp
+++ b/Engine/Plugins/Animation/AnimationWarping/Source/Runtime/Private/BoneControllers/AnimNode_OffsetRootBone.cpp
@@ -270,7 +270,7 @@ void FAnimNode_OffsetRootBone::Evaluate_AnyThread(FPoseContext& Output)
 	{
 		static const TCHAR* LogName = TEXT("OffsetRootBone");
 		const float InnerCircleRadius = 40.0f;
-		const float CircleThickness = 1.5f;
+		const uint16 CircleThickness = 2;
 		const FVector CircleOffset(0,0,1);
 
 		const FTransform TargetBoneInitialTransformWorld = InputBoneTransform * ComponentTransform;
