Update orientation warping to access root bone transform from root offset node

diff --git a/Engine/Plugins/Animation/AnimationWarping/Source/Runtime/Private/BoneControllers/AnimNode_OrientationWarping.cpp b/Engine/Plugins/Animation/AnimationWarping/Source/Runtime/Private/BoneControllers/AnimNode_OrientationWarping.cpp
--- a/Engine/Plugins/Animation/AnimationWarping/Source/Runtime/Private/BoneControllers/AnimNode_OrientationWarping.cpp
+++ b/Engine/Plugins/Animation/AnimationWarping/Source/Runtime/Private/BoneControllers/AnimNode_OrientationWarping.cpp
@@ -5,6 +5,7 @@
 #include "Animation/AnimInstanceProxy.h"
 #include "Animation/AnimNodeFunctionRef.h"
 #include "Animation/AnimRootMotionProvider.h"
+#include "BoneControllers/AnimNode_OffsetRootBone.h"
 #include "HAL/IConsoleManager.h"
 #include "Animation/AnimTrace.h"
 #include "Logging/LogVerbosity.h"
@@ -110,9 +111,9 @@ void FAnimNode_OrientationWarping::UpdateInternal(const FAnimationUpdateContext&
 
 	if (WarpingSpace == EOrientationWarpingSpace::RootBoneTransform)
 	{
-		if (USkeletalMeshComponent* SkelMeshComponent = Context.AnimInstanceProxy->GetSkelMeshComponent())
+		if (UE::AnimationWarping::FRootOffsetProvider* RootOffsetProvider = Context.GetMessage<UE::AnimationWarping::FRootOffsetProvider>())
 		{
-			WarpingSpaceTransform = SkelMeshComponent->GetBoneTransform(0);
+			WarpingSpaceTransform = RootOffsetProvider->GetRootTransform();
 		}
 		else
 		{
