Fix for orientation warping during jumps
#rnx

diff --git a/Engine/Plugins/Animation/AnimationWarping/Source/Runtime/Private/BoneControllers/AnimNode_OrientationWarping.cpp b/Engine/Plugins/Animation/AnimationWarping/Source/Runtime/Private/BoneControllers/AnimNode_OrientationWarping.cpp
--- a/Engine/Plugins/Animation/AnimationWarping/Source/Runtime/Private/BoneControllers/AnimNode_OrientationWarping.cpp
+++ b/Engine/Plugins/Animation/AnimationWarping/Source/Runtime/Private/BoneControllers/AnimNode_OrientationWarping.cpp
@@ -215,13 +215,17 @@ void FAnimNode_OrientationWarping::EvaluateSkeletalControl_AnyThread(FComponentS
 				LocomotionForward = SkeletalMeshRelativeRotation.UnrotateVector(LocomotionRotation.GetForwardVector()).GetSafeNormal();
 			}
 
+			// Flatten locomotion direction, along the rotation axis.
+			LocomotionForward = LocomotionForward - RotationAxisVector.Dot(LocomotionForward) * RotationAxisVector;
+			LocomotionForward.Normalize();
+
 			// @todo: Graph mode using a "manual value" makes no sense. Restructure logic to address this in the future.
 			if (bUseManualRootMotionVelocity)
 			{
 				RootMotionTransformDelta.SetTranslation(ManualRootMotionVelocity * DeltaSeconds);
 			}
 
-			const FVector RootMotionDeltaTranslation = RootMotionTransformDelta.GetTranslation();
+			FVector RootMotionDeltaTranslation = RootMotionTransformDelta.GetTranslation();
 			
 			const float RootMotionDeltaSpeed = RootMotionDeltaTranslation.Size() / DeltaSeconds;
 			if (RootMotionDeltaSpeed < MinRootMotionSpeedThreshold)
@@ -231,6 +235,9 @@ void FAnimNode_OrientationWarping::EvaluateSkeletalControl_AnyThread(FComponentS
 			}
 			else
 			{
+				// Flatten root motion translation, along the rotation axis.
+				RootMotionDeltaTranslation = RootMotionDeltaTranslation - RotationAxisVector.Dot(RootMotionDeltaTranslation) * RotationAxisVector;
+
 				const FVector PreviousRootMotionDeltaDirection = RootMotionDeltaDirection;
 				// Hold previous direction if we can't calculate it from current move delta, because the root is no longer moving
 				RootMotionDeltaDirection = RootMotionDeltaTranslation.GetSafeNormal(UE_SMALL_NUMBER, PreviousRootMotionDeltaDirection);
