Fix FootPlacement Disable Leg Curve implementation

diff --git a/Engine/Plugins/Animation/AnimationWarping/Source/Runtime/Private/BoneControllers/AnimNode_FootPlacement.cpp b/Engine/Plugins/Animation/AnimationWarping/Source/Runtime/Private/BoneControllers/AnimNode_FootPlacement.cpp
--- a/Engine/Plugins/Animation/AnimationWarping/Source/Runtime/Private/BoneControllers/AnimNode_FootPlacement.cpp
+++ b/Engine/Plugins/Animation/AnimationWarping/Source/Runtime/Private/BoneControllers/AnimNode_FootPlacement.cpp
@@ -874,6 +874,8 @@ UE::Anim::FootPlacement::FPlantResult FAnimNode_FootPlacement::FinalizeFootAlign
 			CorrectedFootTransformCS.SetLocation(NotHyperextendedPlantLocation);
 		}
 	}
+	
+	CorrectedFootTransformCS.BlendWith(LegData.InputPose.FootFKTransformCS, LegData.InputPose.DisableLeg);
 
 	check(!CorrectedFootTransformCS.ContainsNaN());
 
@@ -993,6 +995,7 @@ void FAnimNode_FootPlacement::DrawVLog(
 	
 	TRACE_ANIM_NODE_VALUE(Context.CSPContext, TStringBuilder<256>().Append("FootSpeed - ").Append(FString::FromInt(LegData.Idx)).ToString(), LegData.InputPose.Speed);
 	TRACE_ANIM_NODE_VALUE(Context.CSPContext, TStringBuilder<256>().Append("DistanceToPlant - ").Append(FString::FromInt(LegData.Idx)).ToString(), LegData.InputPose.DistanceToPlant);
+	TRACE_ANIM_NODE_VALUE(Context.CSPContext, TStringBuilder<256>().Append("DisableLeg - ").Append(FString::FromInt(LegData.Idx)).ToString(), LegData.InputPose.DisableLeg);
 
 }
 #endif
@@ -1223,6 +1226,9 @@ void FAnimNode_FootPlacement::EvaluateSkeletalControl_AnyThread(FComponentSpaceP
 
 		const UE::Anim::FootPlacement::FPlantResult PlantResult =
 			FinalizeFootAlignment(FootPlacementContext, LegData, LegDef, PelvisTransformCS);
+
+	
+		
 		OutBoneTransforms.Add(PlantResult.FootTranformCS);
 		//OutBoneTransforms.Add(PlantResult.BallTransformCS);
 		//OutBoneTransforms.Add(PlantResult.HipTransformCS);
@@ -1396,6 +1402,7 @@ void FAnimNode_FootPlacement::InitializeBoneReferences(const FBoneContainer& Req
 
 		LegData.SpeedCurveName = LegDef.SpeedCurveName;
 		LegData.DisableLockCurveName = LegDef.DisableLockCurveName;
+		LegData.DisableLegCurveName = LegDef.DisableLegCurveName;
 	}
 
 	PelvisBone.Initialize(RequiredBones);
@@ -1426,7 +1433,7 @@ void FAnimNode_FootPlacement::GatherLegDataFromInputs(
 {
 	FVector LastBallLocation = LegData.InputPose.BallTransformCS.GetLocation();
 
-	const FTransform FootFKTransformCS =
+	LegData.InputPose.FootFKTransformCS =
 		Context.CSPContext.Pose.GetComponentSpaceTransform(LegData.Bones.FKIndex);
 	const FTransform BallTransformCS =
 		Context.CSPContext.Pose.GetComponentSpaceTransform(LegData.Bones.BallIndex);
@@ -1435,11 +1442,11 @@ void FAnimNode_FootPlacement::GatherLegDataFromInputs(
 		Context.CSPContext.Pose.GetComponentSpaceTransform(LegData.Bones.IKIndex);
 	LegData.InputPose.HipTransformCS =
 		Context.CSPContext.Pose.GetComponentSpaceTransform(LegData.Bones.HipIndex);
-
+	
 	LegData.InputPose.BallToFoot =
-		FootFKTransformCS.GetRelativeTransform(BallTransformCS);
+		LegData.InputPose.FootFKTransformCS.GetRelativeTransform(BallTransformCS);
 	LegData.InputPose.FootToBall =
-		BallTransformCS.GetRelativeTransform(FootFKTransformCS);
+		BallTransformCS.GetRelativeTransform(LegData.InputPose.FootFKTransformCS);
 
 	// Can't use ball transform as-is as the foot's IK bone may not be at the FK bone.
 	// Assume the ball is at the same relative position
@@ -1484,6 +1491,8 @@ void FAnimNode_FootPlacement::GatherLegDataFromInputs(
 			Context.CSPContext.Curve.Get(LegData.SpeedCurveName, bValidSpeedCurve, DefaultSpeedCurveValue);
 	}
 
+	LegData.InputPose.DisableLeg = Context.CSPContext.Curve.Get(LegData.DisableLegCurveName);
+
 	// Grab the lock curve's alpha. If the curve isn't set, then LockAlpha is full weight.
 	LegData.InputPose.LockAlpha = 1.0f - Context.CSPContext.Curve.Get(LegData.DisableLockCurveName);
 
@@ -1765,9 +1774,6 @@ void FAnimNode_FootPlacement::ProcessFootAlignment(
 
 	
 	LegData.AlignedFootTransformCS = BlendedPlantTransformCS;
-	
-	float DisableLegCurveValue = Context.CSPContext.Curve.Get(PelvisSettings.DisablePelvisCurveName);
-	LegData.AlignedFootTransformCS.BlendWith(LegData.InputPose.FootTransformCS, DisableLegCurveValue);
 }
 
 FVector FAnimNode_FootPlacement::GetApproachDirWS(const FAnimationBaseContext& Context) const
diff --git a/Engine/Plugins/Animation/AnimationWarping/Source/Runtime/Public/BoneControllers/AnimNode_FootPlacement.h b/Engine/Plugins/Animation/AnimationWarping/Source/Runtime/Public/BoneControllers/AnimNode_FootPlacement.h
--- a/Engine/Plugins/Animation/AnimationWarping/Source/Runtime/Public/BoneControllers/AnimNode_FootPlacement.h
+++ b/Engine/Plugins/Animation/AnimationWarping/Source/Runtime/Public/BoneControllers/AnimNode_FootPlacement.h
@@ -52,10 +52,12 @@ namespace UE::Anim::FootPlacement
 
 		FName SpeedCurveName = NAME_None;
 		FName DisableLockCurveName = NAME_None;
+		FName DisableLegCurveName = NAME_None;
 
 		// Helper struct to store values coming directly, or trivial to calculate from just the input pose.
 		struct FInputPoseData
 		{
+			FTransform FootFKTransformCS = FTransform::Identity;
 			FTransform FootTransformCS = FTransform::Identity;
 			FTransform BallTransformCS = FTransform::Identity;
 			FTransform HipTransformCS = FTransform::Identity;
@@ -68,6 +70,7 @@ namespace UE::Anim::FootPlacement
 			FTransform BallToGround = FTransform::Identity;
 #endif
 			float Speed = 0.0f;
+			float DisableLeg = 0.0f;
 			float LockAlpha = 0.0f;
 			float DistanceToPlant = 0.0f;
 			// Calculated from a range of toe speeds to define when to blend in/out ground rotational alignment
