Orientation Warping: Changed to safe normalization for edge cases.
#rnx

diff --git a/Engine/Plugins/Animation/AnimationWarping/Source/Runtime/Private/BoneControllers/AnimNode_OrientationWarping.cpp b/Engine/Plugins/Animation/AnimationWarping/Source/Runtime/Private/BoneControllers/AnimNode_OrientationWarping.cpp
--- a/Engine/Plugins/Animation/AnimationWarping/Source/Runtime/Private/BoneControllers/AnimNode_OrientationWarping.cpp
+++ b/Engine/Plugins/Animation/AnimationWarping/Source/Runtime/Private/BoneControllers/AnimNode_OrientationWarping.cpp
@@ -216,8 +216,7 @@ void FAnimNode_OrientationWarping::EvaluateSkeletalControl_AnyThread(FComponentS
 			}
 
 			// Flatten locomotion direction, along the rotation axis.
-			LocomotionForward = LocomotionForward - RotationAxisVector.Dot(LocomotionForward) * RotationAxisVector;
-			LocomotionForward.Normalize();
+			LocomotionForward = (LocomotionForward - RotationAxisVector.Dot(LocomotionForward) * RotationAxisVector).GetSafeNormal();
 
 			// @todo: Graph mode using a "manual value" makes no sense. Restructure logic to address this in the future.
 			if (bUseManualRootMotionVelocity)
