Fixed root motion attribute used in animation warping module being "none" due to static initialization dependencies

#rnx minor

diff --git a/Engine/Plugins/Animation/AnimationWarping/Source/Runtime/Private/AnimationWarpingRuntimeModule.cpp b/Engine/Plugins/Animation/AnimationWarping/Source/Runtime/Private/AnimationWarpingRuntimeModule.cpp
--- a/Engine/Plugins/Animation/AnimationWarping/Source/Runtime/Private/AnimationWarpingRuntimeModule.cpp
+++ b/Engine/Plugins/Animation/AnimationWarping/Source/Runtime/Private/AnimationWarpingRuntimeModule.cpp
@@ -14,12 +14,12 @@ class FAnimationWarpingRuntimeModule : public IAnimationWarpingRuntimeModule
 
 namespace UE { namespace AnimationWarping {
 
-// Root motion animation attribute data will always be accessible on the root bone's attribute container
-static const UE::Anim::FAttributeId RootMotionAttributeId = { UE::Anim::IAnimRootMotionProvider::AttributeName , FCompactPoseBoneIndex(0) };
-
 class FModule : public IModuleInterface, public UE::Anim::IAnimRootMotionProvider
 {
 public:
+
+	FModule();
+	
 	virtual void StartupModule() override;
 	virtual void ShutdownModule() override;
 
@@ -28,8 +28,16 @@ public:
 	virtual bool OverrideRootMotion(const FTransform& RootMotionDelta, UE::Anim::FStackAttributeContainer& OutAttributes) const override;
 	virtual bool ExtractRootMotion(const UE::Anim::FStackAttributeContainer& Attributes, FTransform& OutRootMotionDelta) const override;
 	virtual bool HasRootMotion(const UE::Anim::FStackAttributeContainer& Attributes) const override;
+
+	// Root motion animation attribute data will always be accessible on the root bone's attribute container
+	const UE::Anim::FAttributeId RootMotionAttributeId;
 };
 
+FModule::FModule()
+	: RootMotionAttributeId(UE::Anim::IAnimRootMotionProvider::AttributeName , FCompactPoseBoneIndex(0))
+{
+}
+
 void FModule::StartupModule()
 {
 	IModularFeatures::Get().RegisterModularFeature(UE::Anim::IAnimRootMotionProvider::ModularFeatureName, this);
