Update separating plane feature in Foot Placement to only apply to unplanted feet, and use spring based smoothing

diff --git a/Engine/Plugins/Animation/AnimationWarping/Source/Runtime/Private/BoneControllers/AnimNode_FootPlacement.cpp b/Engine/Plugins/Animation/AnimationWarping/Source/Runtime/Private/BoneControllers/AnimNode_FootPlacement.cpp
--- a/Engine/Plugins/Animation/AnimationWarping/Source/Runtime/Private/BoneControllers/AnimNode_FootPlacement.cpp
+++ b/Engine/Plugins/Animation/AnimationWarping/Source/Runtime/Private/BoneControllers/AnimNode_FootPlacement.cpp
@@ -1705,11 +1705,35 @@ void FAnimNode_FootPlacement::ProcessFootAlignment(
 		const float DistanceToSeparatingPlane =
 			UE::Anim::FootPlacement::GetDistanceToPlaneAlongDirection(FootUnalignedLocationCS, SeparatingPlane, -PlaneNormal);
 
-		if (DistanceToSeparatingPlane < 0.0f)
+		if (Plant.PlantType == UE::Anim::FootPlacement::EPlantType::Unplanted)
 		{
-			FootUnalignedLocationCS -= PlaneNormal * DistanceToSeparatingPlane;
+			FVector SeparatingPlaneOffset(FVector::ZeroVector);
+			if (DistanceToSeparatingPlane < 0.0f)
+			{
+				SeparatingPlaneOffset = -PlaneNormal * DistanceToSeparatingPlane ;
+			}
+
+			if (InterpolationSettings.bEnableSeparationInterpolation)
+			{
+				Interpolation.SeparatingPlaneOffset = UKismetMathLibrary::VectorSpringInterp(
+					Interpolation.SeparatingPlaneOffset, SeparatingPlaneOffset, Interpolation.SeparatingPlaneOffsetSpringState,
+					InterpolationSettings.FloorLinearStiffness,
+					InterpolationSettings.FloorLinearDamping,
+					Context.UpdateDeltaTime, 1.0f, 0.0f);
+			}
+			else
+			{
+				Interpolation.SeparatingPlaneOffset = SeparatingPlaneOffset;
+			}
+			
+			FootUnalignedLocationCS  += Interpolation.SeparatingPlaneOffset;
 			FootUnalignedTransformCS.SetLocation(FootUnalignedLocationCS);
 		}
+		else
+		{
+			Interpolation.SeparatingPlaneOffset = FVector::ZeroVector;
+			Interpolation.SeparatingPlaneOffsetSpringState.Reset();
+		}
 
 #if ENABLE_FOOTPLACEMENT_DEBUG
 		if (CVarAnimNodeFootPlacementDebug.GetValueOnAnyThread())
diff --git a/Engine/Plugins/Animation/AnimationWarping/Source/Runtime/Public/BoneControllers/AnimNode_FootPlacement.h b/Engine/Plugins/Animation/AnimationWarping/Source/Runtime/Public/BoneControllers/AnimNode_FootPlacement.h
--- a/Engine/Plugins/Animation/AnimationWarping/Source/Runtime/Public/BoneControllers/AnimNode_FootPlacement.h
+++ b/Engine/Plugins/Animation/AnimationWarping/Source/Runtime/Public/BoneControllers/AnimNode_FootPlacement.h
@@ -106,7 +106,10 @@ namespace UE::Anim::FootPlacement
 		{
 			// Interpolated foot lock offset
 			FTransform UnalignedFootOffsetCS = FTransform::Identity;
-			// Foot lock pring states
+			// Separating plane spring states
+			FVectorSpringState SeparatingPlaneOffsetSpringState;
+			FVector SeparatingPlaneOffset = FVector::ZeroVector;
+			// Foot lock spring states
 			FVectorSpringState PlantOffsetTranslationSpringState;
 			FQuaternionSpringState PlantOffsetRotationSpringState;
 			// Ground alignment spring states
@@ -232,6 +235,12 @@ public:
 
 	UPROPERTY(EditAnywhere, Category = "Plant Settings")
 	float UnplantAngularDamping = 1.0f;
+	
+	UPROPERTY(EditAnywhere, Category = "Plant Settings", meta=(EditCondition="bEnableSeparationInterpolation", DisplayAfter="bEnableSeparationInterpolation"))
+	float SeparationStiffness = 1000.0f;
+
+	UPROPERTY(EditAnywhere, Category = "Plant Settings", meta=(EditCondition="bEnableSeparationInterpolation", DisplayAfter="bEnableSeparationInterpolation"))
+	float SeparationDamping = 1.0f;
 
 	UPROPERTY(EditAnywhere, Category = "Plant Settings", meta=(EditCondition="bEnableFloorInterpolation", DisplayAfter="bEnableFloorInterpolation"))
 	float FloorLinearStiffness = 1000.0f;
@@ -247,6 +256,9 @@ public:
 
 	UPROPERTY(EditAnywhere, Category = "Plant Settings")
 	bool bEnableFloorInterpolation = true;
+	
+	UPROPERTY(EditAnywhere, Category = "Plant Settings")
+	bool bEnableSeparationInterpolation = true;
 };
 
 USTRUCT(BlueprintType)
@@ -472,7 +484,7 @@ public:
 	// Value of 0 disables this
 	UPROPERTY(EditAnywhere, Category = "Plant Settings", meta = (ClampMin = "0.0", UIMin = "0.0"))
 	float SeparatingDistance = 0.0f;
-
+	
 	// Speed at which we transition to fully unplanted.
 	// The range between SpeedThreshold and UnalignmentSpeedThreshold should roughly represent the roll-phase of the foot
 	// TODO: This feels innaccurate most of the time, and varies depending on anim speed. Improve this
