Foot Placement: Disable simple collision hybrid by default

diff --git a/Engine/Plugins/Animation/AnimationWarping/Source/Runtime/Public/BoneControllers/AnimNode_FootPlacement.h b/Engine/Plugins/Animation/AnimationWarping/Source/Runtime/Public/BoneControllers/AnimNode_FootPlacement.h
--- a/Engine/Plugins/Animation/AnimationWarping/Source/Runtime/Public/BoneControllers/AnimNode_FootPlacement.h
+++ b/Engine/Plugins/Animation/AnimationWarping/Source/Runtime/Public/BoneControllers/AnimNode_FootPlacement.h
@@ -294,7 +294,7 @@ public:
 	// Tracing against simple geometry (i.e. it's common for stairs to have simplified ramp collisions) can provide a 
 	// smoother trajectory when the foot is in flight
 	UPROPERTY(EditAnywhere, Category = "Trace Settings", meta = (EditCondition = "bEnabled", DisplayAfter = "bEnabled"))
-	float SimpleCollisionInfluence = 0.5f;
+	float SimpleCollisionInfluence = 0.0f;
 
 	// The channel to use for our simple trace
 	UPROPERTY(EditAnywhere, Category = "Trace Settings", meta = (EditCondition = "bEnabled", DisplayAfter = "bEnabled"))
