add option to disable foot placement pelvis adjustment in air

diff --git a/Engine/Plugins/Animation/AnimationWarping/Source/Runtime/Private/BoneControllers/AnimNode_FootPlacement.cpp b/Engine/Plugins/Animation/AnimationWarping/Source/Runtime/Private/BoneControllers/AnimNode_FootPlacement.cpp
--- a/Engine/Plugins/Animation/AnimationWarping/Source/Runtime/Private/BoneControllers/AnimNode_FootPlacement.cpp
+++ b/Engine/Plugins/Animation/AnimationWarping/Source/Runtime/Private/BoneControllers/AnimNode_FootPlacement.cpp
@@ -874,8 +874,13 @@ UE::Anim::FootPlacement::FPlantResult FAnimNode_FootPlacement::FinalizeFootAlign
 			CorrectedFootTransformCS.SetLocation(NotHyperextendedPlantLocation);
 		}
 	}
-	
-	CorrectedFootTransformCS.BlendWith(LegData.InputPose.FootFKTransformCS, LegData.InputPose.DisableLeg);
+
+	if (LegData.InputPose.DisableLeg > 0)
+	{
+		FTransform DisabledLegTransform = LegData.InputPose.FootFKTransformCS;
+		DisabledLegTransform.SetTranslation(DisabledLegTransform.GetTranslation() + PelvisData.Interpolation.PelvisTranslationOffset * (1.0-PelvisData.DisablePelvis));
+		CorrectedFootTransformCS.BlendWith(DisabledLegTransform, LegData.InputPose.DisableLeg);
+	}
 
 	check(!CorrectedFootTransformCS.ContainsNaN());
 
@@ -1200,9 +1205,6 @@ void FAnimNode_FootPlacement::EvaluateSkeletalControl_AnyThread(FComponentSpaceP
 	// Based on the ground alignment, search for the best Pelvis transform
 	FTransform PelvisTransformCS = SolvePelvis(FootPlacementContext);
 
-	float DisablePelvisCurveValue = Output.Curve.Get(PelvisSettings.DisablePelvisCurveName);
-	PelvisTransformCS.BlendWith(PelvisData.InputPose.FKTransformCS, DisablePelvisCurveValue);
-
 #if ENABLE_FOOTPLACEMENT_DEBUG
 	const FTransform PelvisTargetTransformCS = PelvisTransformCS;
 	if (CVarAnimNodeFootPlacementDebug.GetValueOnAnyThread())
@@ -1216,6 +1218,10 @@ void FAnimNode_FootPlacement::EvaluateSkeletalControl_AnyThread(FComponentSpaceP
 	{
 		PelvisTransformCS = UpdatePelvisInterpolation(FootPlacementContext, PelvisTransformCS);
 	}
+	
+	PelvisData.DisablePelvis = Output.Curve.Get(PelvisSettings.DisablePelvisCurveName);
+	PelvisTransformCS.BlendWith(PelvisData.InputPose.FKTransformCS, PelvisData.DisablePelvis);
+	
 	check(!PelvisTransformCS.ContainsNaN());
 	OutBoneTransforms.Add(FBoneTransform(PelvisData.Bones.FkBoneIndex, PelvisTransformCS));
 
@@ -1527,7 +1533,8 @@ void FAnimNode_FootPlacement::ProcessCharacterState(const UE::Anim::FootPlacemen
 			(Context.MovementComponent->MovementMode == MOVE_NavWalking)) &&
 		Context.MovementComponent->CurrentFloor.bBlockingHit;
 
-	if (CharacterData.bIsOnGround && bWasOnGround && (PelvisSettings.ActorMovementCompensationMode != EActorMovementCompensationMode::ComponentSpace))
+	bool bOnGround =  !PelvisSettings.bDisablePelvisOffsetInAir || (CharacterData.bIsOnGround && bWasOnGround);
+	if (bOnGround && PelvisSettings.ActorMovementCompensationMode != EActorMovementCompensationMode::ComponentSpace)
 	{
 		FVector OwningComponentAdjustedLastLocationWS;
 		if (PelvisSettings.ActorMovementCompensationMode == EActorMovementCompensationMode::SuddenMotionOnly)
diff --git a/Engine/Plugins/Animation/AnimationWarping/Source/Runtime/Public/BoneControllers/AnimNode_FootPlacement.h b/Engine/Plugins/Animation/AnimationWarping/Source/Runtime/Public/BoneControllers/AnimNode_FootPlacement.h
--- a/Engine/Plugins/Animation/AnimationWarping/Source/Runtime/Public/BoneControllers/AnimNode_FootPlacement.h
+++ b/Engine/Plugins/Animation/AnimationWarping/Source/Runtime/Public/BoneControllers/AnimNode_FootPlacement.h
@@ -154,6 +154,9 @@ namespace UE::Anim::FootPlacement
 			FVector PelvisTranslationOffset = FVector::ZeroVector;
 			FVectorSpringState PelvisTranslationSpringState;
 		} Interpolation;
+
+		
+		float DisablePelvis = 0.0f;
 	};
 
 	struct FCharacterData
@@ -394,6 +397,9 @@ public:
 
 	UPROPERTY(EditAnywhere, Category = "Pelvis Settings")
 	bool bEnableInterpolation = true;
+
+	UPROPERTY(EditAnywhere, Category = "Pelvis Settings")
+	bool bDisablePelvisOffsetInAir = true;
 	
 	UPROPERTY(EditAnywhere, Category = "Pelvis Settings")
 	FName DisablePelvisCurveName = NAME_None;
