Add a shape cast to OffsetRootBone to prevent crashing with walls

diff --git a/Engine/Plugins/Animation/AnimationWarping/Source/Runtime/Private/BoneControllers/AnimNode_OffsetRootBone.cpp b/Engine/Plugins/Animation/AnimationWarping/Source/Runtime/Private/BoneControllers/AnimNode_OffsetRootBone.cpp
--- a/Engine/Plugins/Animation/AnimationWarping/Source/Runtime/Private/BoneControllers/AnimNode_OffsetRootBone.cpp
+++ b/Engine/Plugins/Animation/AnimationWarping/Source/Runtime/Private/BoneControllers/AnimNode_OffsetRootBone.cpp
@@ -1,6 +1,7 @@
 // Copyright Epic Games, Inc. All Rights Reserved.
 
 #include "BoneControllers/AnimNode_OffsetRootBone.h"
+
 #include "Animation/AnimInstanceProxy.h"
 #include "Animation/AnimNodeFunctionRef.h"
 #include "Animation/AnimRootMotionProvider.h"
@@ -125,6 +126,53 @@ void FAnimNode_OffsetRootBone::Evaluate_AnyThread(FPoseContext& Output)
 		// Apply the offset as is (component space) in manual mode
 		RootMotionTransformDelta = FTransform(GetRotationDelta(), GetTranslationDelta());
 	}
+	float MaxTranslationOffset = GetMaxTranslationError();
+
+	bool bCollisionDetected = false;
+	FVector CollisionPoint;
+	FVector CollisionNormal;
+
+	if (GetCollisionTestingMode() != ECollisionResponseType::Disabled && MaxTranslationOffset > 0)
+	{
+		const FCollisionShape CollisionShape = FCollisionShape::MakeSphere(GetCollisionTestShapeRadius());
+		
+		FVector TraceDirectionCS = FVector(0,1,0);
+		if (RootMotionTransformDelta.GetTranslation().Length() > 0.1f)
+		{
+			TraceDirectionCS = LastNonZeroRootMotionDirection = RootMotionTransformDelta.GetTranslation().GetUnsafeNormal();
+		}
+		else if (LastNonZeroRootMotionDirection.SquaredLength() > UE_SMALL_NUMBER)
+		{
+			TraceDirectionCS = LastNonZeroRootMotionDirection;
+		}
+		const FVector TraceDirectionWS = SimulatedRotation.RotateVector(TraceDirectionCS);
+		
+		const FVector TraceStart = ComponentTransform.GetTranslation() + GetCollisionTestShapeOffset();
+		const FVector TraceEnd = TraceStart + (MaxTranslationOffset * TraceDirectionWS);
+    
+		FCollisionQueryParams QueryParams;
+		// Ignore self and all attached components
+		QueryParams.AddIgnoredActor(AnimInstanceProxy->GetSkelMeshComponent()->GetOwner());
+    
+		const ECollisionChannel CollisionChannel = UEngineTypes::ConvertToCollisionChannel(TraceTypeQuery1);
+    
+		FHitResult HitResult;
+		const bool bHit = AnimInstanceProxy->GetSkelMeshComponent()->GetWorld()->SweepSingleByChannel(
+			HitResult, TraceStart, TraceEnd, FQuat::Identity, CollisionChannel, CollisionShape, QueryParams);
+    
+		if (bHit && HitResult.Distance < MaxTranslationOffset)
+		{
+			if (GetCollisionTestingMode() == ECollisionResponseType::ShrinkMaxTranslation)
+			{
+				MaxTranslationOffset = HitResult.Distance;
+			}
+			
+			bCollisionDetected = true;
+			CollisionPoint = HitResult.ImpactPoint;
+			CollisionNormal = HitResult.ImpactNormal;
+		}
+	}
+	
 
 	FTransform ConsumedRootMotionDelta;
 
@@ -191,6 +239,25 @@ void FAnimNode_OffsetRootBone::Evaluate_AnyThread(FPoseContext& Output)
 			}
 		}
 
+		if (bCollisionDetected && GetCollisionTestingMode() == ECollisionResponseType::PlanarCollision)
+		{
+			float B = FVector::DotProduct(TranslationOffsetDelta, CollisionNormal);
+			if (B > UE_KINDA_SMALL_NUMBER)
+			{
+				FVector OffsetCollisionPoint = CollisionPoint + CollisionNormal * GetCollisionTestShapeRadius();
+				float CollisionParam = FVector::DotProduct((OffsetCollisionPoint - SimulatedTranslation), CollisionNormal) / B;
+				if (CollisionParam >=0 && CollisionParam < 1)
+				{
+					FVector TranslationToPlane = TranslationOffsetDelta * CollisionParam;
+					FVector TranslationAlongPlane = TranslationOffsetDelta - TranslationToPlane;
+					TranslationAlongPlane = TranslationAlongPlane - FVector::DotProduct(TranslationAlongPlane, CollisionNormal);
+
+					TranslationOffsetDelta = TranslationToPlane + TranslationAlongPlane;
+				}
+			}
+		}
+		
+
 		SimulatedTranslation = SimulatedTranslation + TranslationOffsetDelta;
 	}
 
@@ -225,13 +292,13 @@ void FAnimNode_OffsetRootBone::Evaluate_AnyThread(FPoseContext& Output)
 		SimulatedRotation = RotationOffsetDelta * SimulatedRotation;
 	}
 
-	if (GetMaxTranslationError() >= 0.0f)
+	if (MaxTranslationOffset >= 0.0f)
 	{
 		FVector TranslationOffset = ComponentTransform.GetLocation() - SimulatedTranslation;
 		const float TranslationOffsetSize = TranslationOffset.Size();
-		if (TranslationOffsetSize > GetMaxTranslationError())
+		if (TranslationOffsetSize > MaxTranslationOffset)
 		{
-			TranslationOffset = TranslationOffset.GetClampedToMaxSize(GetMaxTranslationError());
+			TranslationOffset = TranslationOffset.GetClampedToMaxSize(MaxTranslationOffset);
 			SimulatedTranslation = ComponentTransform.GetLocation() - TranslationOffset;
 		}
 	}
@@ -277,10 +344,15 @@ void FAnimNode_OffsetRootBone::Evaluate_AnyThread(FPoseContext& Output)
 		const FTransform TargetBoneTransformWorld = TargetBoneTransform * ComponentTransform;
 		UObject* LogOwner = AnimInstanceProxy->GetAnimInstanceObject();
 
-		if (GetMaxTranslationError() >= 0.0f)
+		if (MaxTranslationOffset >= 0.0f)
 		{
-			const float OuterCircleRadius = GetMaxTranslationError() + InnerCircleRadius;
+			const float OuterCircleRadius = MaxTranslationOffset + InnerCircleRadius;
 			UE_VLOG_CIRCLE_THICK(AnimInstanceProxy->GetAnimInstanceObject(), TEXT("OffsetRootBone"), Display, ComponentTransform.GetLocation() + CircleOffset, FVector::UpVector, OuterCircleRadius, FColor::Red, CircleThickness, TEXT(""));
+			
+			if (bCollisionDetected)
+			{
+				UE_VLOG_CIRCLE_THICK(LogOwner, LogName, Display, CollisionPoint, CollisionNormal, GetCollisionTestShapeRadius(), FColor::Red, CircleThickness, TEXT(""));
+			}
 		}
 		
 		UE_VLOG_CIRCLE_THICK(LogOwner, LogName, Display, ComponentTransform.GetLocation() + CircleOffset, FVector::UpVector, InnerCircleRadius, FColor::Blue, CircleThickness, TEXT(""));
@@ -308,9 +380,9 @@ void FAnimNode_OffsetRootBone::Evaluate_AnyThread(FPoseContext& Output)
 		const FTransform TargetBoneInitialTransformWorld = InputBoneTransform * ComponentTransform;
 		const FTransform TargetBoneTransformWorld = TargetBoneTransform * ComponentTransform;
 
-		if (GetMaxTranslationError() >= 0.0f)
+		if (MaxTranslationOffset >= 0.0f)
 		{
-			const float OuterCircleRadius = GetMaxTranslationError() + InnerCircleRadius;
+			const float OuterCircleRadius = MaxTranslationOffset + InnerCircleRadius;
 			AnimInstanceProxy->AnimDrawDebugCircle(ComponentTransform.GetLocation(), OuterCircleRadius, 36, FColor::Red,
 				FVector::UpVector, false, -1.0f, SDPG_World, CircleThickness);
 		}
@@ -415,6 +487,21 @@ float FAnimNode_OffsetRootBone::GetRotationSpeedRatio() const
 	return GET_ANIM_NODE_DATA(float, RotationSpeedRatio);
 }
 
+ECollisionResponseType FAnimNode_OffsetRootBone::GetCollisionTestingMode() const
+{
+	return GET_ANIM_NODE_DATA(ECollisionResponseType, CollisionTestingMode);
+}
+
+float FAnimNode_OffsetRootBone::GetCollisionTestShapeRadius() const
+{
+	return GET_ANIM_NODE_DATA(float, CollisionTestShapeRadius);
+}
+
+const FVector& FAnimNode_OffsetRootBone::GetCollisionTestShapeOffset() const
+{
+	return GET_ANIM_NODE_DATA(FVector, CollisionTestShapeOffset);
+}
+
 void FAnimNode_OffsetRootBone::Reset(const FAnimationBaseContext& Context)
 {
 	ComponentTransform = Context.AnimInstanceProxy->GetComponentTransform();
diff --git a/Engine/Plugins/Animation/AnimationWarping/Source/Runtime/Public/BoneControllers/AnimNode_OffsetRootBone.h b/Engine/Plugins/Animation/AnimationWarping/Source/Runtime/Public/BoneControllers/AnimNode_OffsetRootBone.h
--- a/Engine/Plugins/Animation/AnimationWarping/Source/Runtime/Public/BoneControllers/AnimNode_OffsetRootBone.h
+++ b/Engine/Plugins/Animation/AnimationWarping/Source/Runtime/Public/BoneControllers/AnimNode_OffsetRootBone.h
@@ -50,6 +50,17 @@ enum class EOffsetRootBoneMode : uint8
 	Release,
 };
 
+UENUM(BlueprintType)
+enum class ECollisionResponseType : uint8
+{
+	// No Collision testing
+	Disabled,
+	// Reduce effective Max Translation offset to prevent penetration with nearby obstacles
+	ShrinkMaxTranslation,
+	// Slide along a plane based on shape cast contact point
+	PlanarCollision,
+};
+
 USTRUCT(BlueprintInternalUseOnly, Experimental)
 struct ANIMATIONWARPINGRUNTIME_API FAnimNode_OffsetRootBone : public FAnimNode_Base
 {
@@ -119,6 +130,15 @@ struct ANIMATIONWARPINGRUNTIME_API FAnimNode_OffsetRootBone : public FAnimNode_B
 	// For procedural values, consider adjusting the input by delta time.
 	UPROPERTY(EditAnywhere, Category = Evaluation, meta = (FoldProperty, PinHiddenByDefault))
 	FRotator RotationDelta = FRotator::ZeroRotator;
+
+
+	UPROPERTY(EditAnywhere, Category = CollisionTesting, meta = (FoldProperty, PinHiddenByDefault))
+	ECollisionResponseType CollisionTestingMode = ECollisionResponseType::Disabled;
+	UPROPERTY(EditAnywhere, Category = CollisionTesting, meta = (EditCondition = "CollisionTestingMode != ECollisionTestingMode::Disabled", DisplayAfter="CollisionTestingMode", FoldProperty, PinHiddenByDefault))
+	float CollisionTestShapeRadius = 30;
+	UPROPERTY(EditAnywhere, Category = CollisionTesting, meta = (EditCondition = "CollisionTestingMode != ECollisionTestingMode::Disabled", DisplayAfter="CollisionTestingMode", FoldProperty, PinHiddenByDefault))
+	FVector CollisionTestShapeOffset = {0,0,60};
+			
 #endif
 
 public:
@@ -144,7 +164,9 @@ public:
 	bool GetClampToRotationVelocity() const;
 	float GetTranslationSpeedRatio() const;
 	float GetRotationSpeedRatio() const;
-
+	ECollisionResponseType GetCollisionTestingMode() const;
+	float GetCollisionTestShapeRadius() const;
+	const FVector& GetCollisionTestShapeOffset() const;
 
 	// get the current simulated root transform
 	void GetOffsetRootTransform(FTransform& OutTransform)
@@ -172,6 +194,8 @@ private:
 	FVector SimulatedTranslation = FVector::ZeroVector;
 	FQuat SimulatedRotation = FQuat::Identity;
 
+	FVector LastNonZeroRootMotionDirection = FVector::ZeroVector;
+
 	FGraphTraversalCounter UpdateCounter;
 };
 
