Fix CIS Cook warning (enum script name confilct in OffsetRootBone)

diff --git a/Engine/Plugins/Animation/AnimationWarping/Source/Runtime/Private/BoneControllers/AnimNode_OffsetRootBone.cpp b/Engine/Plugins/Animation/AnimationWarping/Source/Runtime/Private/BoneControllers/AnimNode_OffsetRootBone.cpp
--- a/Engine/Plugins/Animation/AnimationWarping/Source/Runtime/Private/BoneControllers/AnimNode_OffsetRootBone.cpp
+++ b/Engine/Plugins/Animation/AnimationWarping/Source/Runtime/Private/BoneControllers/AnimNode_OffsetRootBone.cpp
@@ -132,7 +132,7 @@ void FAnimNode_OffsetRootBone::Evaluate_AnyThread(FPoseContext& Output)
 	FVector CollisionPoint;
 	FVector CollisionNormal;
 
-	if (GetCollisionTestingMode() != ECollisionResponseType::Disabled && MaxTranslationOffset > 0)
+	if (GetCollisionTestingMode() != EOffsetRootBone_CollisionTestingMode::Disabled && MaxTranslationOffset > 0)
 	{
 		const FCollisionShape CollisionShape = FCollisionShape::MakeSphere(GetCollisionTestShapeRadius());
 		
@@ -162,7 +162,7 @@ void FAnimNode_OffsetRootBone::Evaluate_AnyThread(FPoseContext& Output)
     
 		if (bHit && HitResult.Distance < MaxTranslationOffset)
 		{
-			if (GetCollisionTestingMode() == ECollisionResponseType::ShrinkMaxTranslation)
+			if (GetCollisionTestingMode() == EOffsetRootBone_CollisionTestingMode::ShrinkMaxTranslation)
 			{
 				MaxTranslationOffset = HitResult.Distance;
 			}
@@ -239,7 +239,7 @@ void FAnimNode_OffsetRootBone::Evaluate_AnyThread(FPoseContext& Output)
 			}
 		}
 
-		if (bCollisionDetected && GetCollisionTestingMode() == ECollisionResponseType::PlanarCollision)
+		if (bCollisionDetected && GetCollisionTestingMode() == EOffsetRootBone_CollisionTestingMode::PlanarCollision)
 		{
 			float B = FVector::DotProduct(TranslationOffsetDelta, CollisionNormal);
 			if (B > UE_KINDA_SMALL_NUMBER)
@@ -487,9 +487,9 @@ float FAnimNode_OffsetRootBone::GetRotationSpeedRatio() const
 	return GET_ANIM_NODE_DATA(float, RotationSpeedRatio);
 }
 
-ECollisionResponseType FAnimNode_OffsetRootBone::GetCollisionTestingMode() const
+EOffsetRootBone_CollisionTestingMode FAnimNode_OffsetRootBone::GetCollisionTestingMode() const
 {
-	return GET_ANIM_NODE_DATA(ECollisionResponseType, CollisionTestingMode);
+	return GET_ANIM_NODE_DATA(EOffsetRootBone_CollisionTestingMode, CollisionTestingMode);
 }
 
 float FAnimNode_OffsetRootBone::GetCollisionTestShapeRadius() const
diff --git a/Engine/Plugins/Animation/AnimationWarping/Source/Runtime/Public/BoneControllers/AnimNode_OffsetRootBone.h b/Engine/Plugins/Animation/AnimationWarping/Source/Runtime/Public/BoneControllers/AnimNode_OffsetRootBone.h
--- a/Engine/Plugins/Animation/AnimationWarping/Source/Runtime/Public/BoneControllers/AnimNode_OffsetRootBone.h
+++ b/Engine/Plugins/Animation/AnimationWarping/Source/Runtime/Public/BoneControllers/AnimNode_OffsetRootBone.h
@@ -51,7 +51,7 @@ enum class EOffsetRootBoneMode : uint8
 };
 
 UENUM(BlueprintType)
-enum class ECollisionResponseType : uint8
+enum class EOffsetRootBone_CollisionTestingMode : uint8
 {
 	// No Collision testing
 	Disabled,
@@ -133,7 +133,7 @@ struct ANIMATIONWARPINGRUNTIME_API FAnimNode_OffsetRootBone : public FAnimNode_B
 
 
 	UPROPERTY(EditAnywhere, Category = CollisionTesting, meta = (FoldProperty, PinHiddenByDefault))
-	ECollisionResponseType CollisionTestingMode = ECollisionResponseType::Disabled;
+	EOffsetRootBone_CollisionTestingMode CollisionTestingMode = EOffsetRootBone_CollisionTestingMode::Disabled;
 	UPROPERTY(EditAnywhere, Category = CollisionTesting, meta = (EditCondition = "CollisionTestingMode != ECollisionTestingMode::Disabled", DisplayAfter="CollisionTestingMode", FoldProperty, PinHiddenByDefault))
 	float CollisionTestShapeRadius = 30;
 	UPROPERTY(EditAnywhere, Category = CollisionTesting, meta = (EditCondition = "CollisionTestingMode != ECollisionTestingMode::Disabled", DisplayAfter="CollisionTestingMode", FoldProperty, PinHiddenByDefault))
@@ -164,7 +164,7 @@ public:
 	bool GetClampToRotationVelocity() const;
 	float GetTranslationSpeedRatio() const;
 	float GetRotationSpeedRatio() const;
-	ECollisionResponseType GetCollisionTestingMode() const;
+	EOffsetRootBone_CollisionTestingMode GetCollisionTestingMode() const;
 	float GetCollisionTestShapeRadius() const;
 	const FVector& GetCollisionTestShapeOffset() const;
 
