Added blueprint callable helper function to get the value of a curve from an animation at a given time.
#rnx

diff --git a/Engine/Plugins/Animation/AnimationWarping/Source/Runtime/Private/AnimationWarpingLibrary.cpp b/Engine/Plugins/Animation/AnimationWarping/Source/Runtime/Private/AnimationWarpingLibrary.cpp
--- a/Engine/Plugins/Animation/AnimationWarping/Source/Runtime/Private/AnimationWarpingLibrary.cpp
+++ b/Engine/Plugins/Animation/AnimationWarping/Source/Runtime/Private/AnimationWarpingLibrary.cpp
@@ -2,6 +2,7 @@
 
 #include "AnimationWarpingLibrary.h"
 #include "BoneControllers/AnimNode_OffsetRootBone.h"
+#include "Animation/AnimMontage.h"
 
 FTransform UAnimationWarpingLibrary::GetOffsetRootTransform(const FAnimNodeReference& Node)
 {
@@ -11,6 +12,42 @@ FTransform UAnimationWarpingLibrary::GetOffsetRootTransform(const FAnimNodeRefer
 	{
 		OffsetRoot->GetOffsetRootTransform(Transform);
 	}
-	
+
 	return Transform;
 }
+
+bool UAnimationWarpingLibrary::GetCurveValueFromAnimation(const UAnimSequenceBase* Animation, FName CurveName, float Time, float& OutValue)
+{
+	OutValue = 0.0f;
+
+	// If Animation is a Montage we need to get the AnimSequence at the desired time
+	// because EvaluateCurveData doesn't work when called from a Montage and the curve is in the AnimSequence within the montage. 
+	if (const UAnimMontage* Montage = Cast<UAnimMontage>(Animation))
+	{
+		// For now just assume we are working with a montage with a single slot, which is the most common case anyway
+		// The engine also makes this assumption some times. E.g Root motion is only extracted from the first track See: UAnimMontage::ExtractRootMotionFromTrackRange
+		if (Montage->SlotAnimTracks.Num() > 0)
+		{
+			const FAnimTrack& AnimTrack = Montage->SlotAnimTracks[0].AnimTrack;
+
+			if (const FAnimSegment* Segment = AnimTrack.GetSegmentAtTime(Time))
+			{
+				if (Segment->GetAnimReference() && Segment->GetAnimReference()->HasCurveData(CurveName))
+				{
+					float ActualTime = Segment->ConvertTrackPosToAnimPos(Time);
+					ActualTime = FMath::Clamp(ActualTime, Segment->AnimStartTime, Segment->AnimEndTime);
+
+					OutValue = Segment->GetAnimReference()->EvaluateCurveData(CurveName, ActualTime);
+					return true;
+				}
+			}
+		}
+	}
+	else if (Animation && Animation->HasCurveData(CurveName))
+	{
+		OutValue = Animation->EvaluateCurveData(CurveName, Time);
+		return true;
+	}
+
+	return false;
+}
diff --git a/Engine/Plugins/Animation/AnimationWarping/Source/Runtime/Public/AnimationWarpingLibrary.h b/Engine/Plugins/Animation/AnimationWarping/Source/Runtime/Public/AnimationWarpingLibrary.h
--- a/Engine/Plugins/Animation/AnimationWarping/Source/Runtime/Public/AnimationWarpingLibrary.h
+++ b/Engine/Plugins/Animation/AnimationWarping/Source/Runtime/Public/AnimationWarpingLibrary.h
@@ -20,4 +20,7 @@ public:
     UFUNCTION(BlueprintPure, Category = "Animation|BlendStack", meta = (BlueprintThreadSafe))
     static FTransform GetOffsetRootTransform(const FAnimNodeReference& Node);
 	
-};
\ No newline at end of file
+	/** Helper function to extract the value of a curve in an animation at a given time */
+	UFUNCTION(BlueprintPure, Category = "Animation", meta = (BlueprintThreadSafe))
+	static bool GetCurveValueFromAnimation(const UAnimSequenceBase* Animation, FName CurveName, float Time, float& OutValue);
+};
