Attribute based root motion support

diff --git a/Engine/Plugins/Animation/AnimationWarping/Source/Editor/Private/AnimGraph/AnimGraphNode_OverrideRootMotion.cpp b/Engine/Plugins/Animation/AnimationWarping/Source/Editor/Private/AnimGraph/AnimGraphNode_OverrideRootMotion.cpp
new file mode 100644
--- /dev/null
+++ b/Engine/Plugins/Animation/AnimationWarping/Source/Editor/Private/AnimGraph/AnimGraphNode_OverrideRootMotion.cpp
@@ -0,0 +1,38 @@
+// Copyright Epic Games, Inc. All Rights Reserved.
+
+#include "AnimGraph/AnimGraphNode_OverrideRootMotion.h"
+#include "Animation/AnimRootMotionProvider.h"
+#include "Kismet2/BlueprintEditorUtils.h"
+
+#define LOCTEXT_NAMESPACE "AnimationWarping"
+
+UAnimGraphNode_OverrideRootMotion::UAnimGraphNode_OverrideRootMotion(const FObjectInitializer& ObjectInitializer) : Super(ObjectInitializer)
+{
+}
+
+FText UAnimGraphNode_OverrideRootMotion::GetTooltipText() const
+{
+	return LOCTEXT("OverrideRootMotionTooltip", "Overrides the root motion attribute");
+}
+
+FText UAnimGraphNode_OverrideRootMotion::GetNodeTitle(ENodeTitleType::Type TitleType) const
+{
+	return LOCTEXT("OverrideRootMotion", "Override Root Motion");
+}
+
+FLinearColor UAnimGraphNode_OverrideRootMotion::GetNodeTitleColor() const
+{
+	return FLinearColor(FColor(153.f, 0.f, 0.f));
+}
+
+void UAnimGraphNode_OverrideRootMotion::GetInputLinkAttributes(FNodeAttributeArray& OutAttributes) const
+{
+	OutAttributes.Add(UE::Anim::IAnimRootMotionProvider::AttributeName);
+}
+
+void UAnimGraphNode_OverrideRootMotion::GetOutputLinkAttributes(FNodeAttributeArray& OutAttributes) const
+{
+	OutAttributes.Add(UE::Anim::IAnimRootMotionProvider::AttributeName);
+}
+
+#undef LOCTEXT_NAMESPACE
diff --git a/Engine/Plugins/Animation/AnimationWarping/Source/Editor/Public/AnimGraph/AnimGraphNode_OverrideRootMotion.h b/Engine/Plugins/Animation/AnimationWarping/Source/Editor/Public/AnimGraph/AnimGraphNode_OverrideRootMotion.h
new file mode 100644
--- /dev/null
+++ b/Engine/Plugins/Animation/AnimationWarping/Source/Editor/Public/AnimGraph/AnimGraphNode_OverrideRootMotion.h
@@ -0,0 +1,30 @@
+// Copyright Epic Games, Inc. All Rights Reserved.
+
+#pragma once
+
+#include "AnimGraphNode_Base.h"
+#include "BoneControllers/AnimNode_OverrideRootMotion.h"
+
+#include "AnimGraphNode_OverrideRootMotion.generated.h"
+
+namespace ENodeTitleType { enum Type : int; }
+
+UCLASS(Experimental)
+class ANIMATIONWARPINGEDITOR_API UAnimGraphNode_OverrideRootMotion : public UAnimGraphNode_Base
+{
+	GENERATED_UCLASS_BODY()
+
+
+	UPROPERTY(EditAnywhere, Category = Settings)
+	FAnimNode_OverrideRootMotion Node;
+
+public:
+	// UEdGraphNode interface
+	virtual FText GetNodeTitle(ENodeTitleType::Type TitleType) const override;
+	virtual FLinearColor GetNodeTitleColor() const override;
+	virtual FText GetTooltipText() const override;
+	virtual void GetInputLinkAttributes(FNodeAttributeArray& OutAttributes) const override;
+	virtual void GetOutputLinkAttributes(FNodeAttributeArray& OutAttributes) const override;
+	// End of UEdGraphNode interface
+
+};
diff --git a/Engine/Plugins/Animation/AnimationWarping/Source/Runtime/Private/BoneControllers/AnimNode_OffsetRootBone.cpp b/Engine/Plugins/Animation/AnimationWarping/Source/Runtime/Private/BoneControllers/AnimNode_OffsetRootBone.cpp
--- a/Engine/Plugins/Animation/AnimationWarping/Source/Runtime/Private/BoneControllers/AnimNode_OffsetRootBone.cpp
+++ b/Engine/Plugins/Animation/AnimationWarping/Source/Runtime/Private/BoneControllers/AnimNode_OffsetRootBone.cpp
@@ -71,7 +71,8 @@ void FAnimNode_OffsetRootBone::Update_AnyThread(const FAnimationUpdateContext& C
 	GetEvaluateGraphExposedInputs().Execute(Context);
 
 	// If we just became relevant and haven't been initialized yet, then reset.
-	if (!bIsFirstUpdate && UpdateCounter.HasEverBeenUpdated() && !UpdateCounter.WasSynchronizedCounter(Context.AnimInstanceProxy->GetUpdateCounter()))
+	if (GetResetEveryFrame() || 
+		(!bIsFirstUpdate && UpdateCounter.HasEverBeenUpdated() && !UpdateCounter.WasSynchronizedCounter(Context.AnimInstanceProxy->GetUpdateCounter())))
 	{
 		Reset(Context);
 	}
@@ -429,6 +430,11 @@ EWarpingEvaluationMode FAnimNode_OffsetRootBone::GetEvaluationMode() const
 	return GET_ANIM_NODE_DATA(EWarpingEvaluationMode, EvaluationMode);
 }
 
+bool FAnimNode_OffsetRootBone::GetResetEveryFrame() const
+{
+	return GET_ANIM_NODE_DATA(bool, bResetEveryFrame);
+}
+
 const FVector& FAnimNode_OffsetRootBone::GetTranslationDelta() const
 {
 	return GET_ANIM_NODE_DATA(FVector, TranslationDelta);
diff --git a/Engine/Plugins/Animation/AnimationWarping/Source/Runtime/Private/BoneControllers/AnimNode_OverrideRootMotion.cpp b/Engine/Plugins/Animation/AnimationWarping/Source/Runtime/Private/BoneControllers/AnimNode_OverrideRootMotion.cpp
new file mode 100644
--- /dev/null
+++ b/Engine/Plugins/Animation/AnimationWarping/Source/Runtime/Private/BoneControllers/AnimNode_OverrideRootMotion.cpp
@@ -0,0 +1,69 @@
+// Copyright Epic Games, Inc. All Rights Reserved.
+
+#include "BoneControllers/AnimNode_OverrideRootMotion.h"
+
+#include "Animation/AnimInstanceProxy.h"
+#include "Animation/AnimNodeFunctionRef.h"
+#include "Animation/AnimRootMotionProvider.h"
+#include "HAL/IConsoleManager.h"
+#include "VisualLogger/VisualLogger.h"
+
+#include UE_INLINE_GENERATED_CPP_BY_NAME(AnimNode_OverrideRootMotion)
+
+DECLARE_CYCLE_STAT(TEXT("OverrideRootMotion Eval"), STAT_OverrideRootMotion_Eval, STATGROUP_Anim);
+
+void FAnimNode_OverrideRootMotion::Initialize_AnyThread(const FAnimationInitializeContext& Context)
+{
+	Super::Initialize_AnyThread(Context);
+	AnimInstanceProxy = Context.AnimInstanceProxy;
+	Source.Initialize(Context);
+}
+
+void FAnimNode_OverrideRootMotion::CacheBones_AnyThread(const FAnimationCacheBonesContext& Context)
+{
+	Super::CacheBones_AnyThread(Context);
+	Source.CacheBones(Context);
+}
+
+void FAnimNode_OverrideRootMotion::Update_AnyThread(const FAnimationUpdateContext& Context)
+{
+	Super::Update_AnyThread(Context);
+
+	GetEvaluateGraphExposedInputs().Execute(Context);
+
+	DeltaTime = Context.GetDeltaTime();
+
+	Source.Update(Context);
+}
+
+void FAnimNode_OverrideRootMotion::Evaluate_AnyThread(FPoseContext& Output)
+{
+	SCOPE_CYCLE_COUNTER(STAT_OverrideRootMotion_Eval);
+
+	Super::Evaluate_AnyThread(Output);
+	Source.Evaluate(Output);
+
+	const UE::Anim::IAnimRootMotionProvider* RootMotionProvider = UE::Anim::IAnimRootMotionProvider::Get();
+
+	FTransform RootMotionTransform(FTransform::Identity);
+	RootMotionProvider->ExtractRootMotion(Output.CustomAttributes, RootMotionTransform);
+
+	FVector Velocity = GetOverrideVelocity();
+	FVector Delta = Velocity * DeltaTime;
+	FTransform Transform(Output.AnimInstanceProxy->GetComponentTransform().InverseTransformVector(Delta));
+
+	RootMotionTransform.BlendWith(Transform, GetAlpha());
+	
+	RootMotionProvider->OverrideRootMotion(RootMotionTransform, Output.CustomAttributes);
+}
+
+const FVector& FAnimNode_OverrideRootMotion::GetOverrideVelocity() const
+{
+	return GET_ANIM_NODE_DATA(FVector, OverrideVelocity);
+}
+
+float FAnimNode_OverrideRootMotion::GetAlpha() const
+{
+	return GET_ANIM_NODE_DATA(float, Alpha);
+}
+
diff --git a/Engine/Plugins/Animation/AnimationWarping/Source/Runtime/Public/BoneControllers/AnimNode_OffsetRootBone.h b/Engine/Plugins/Animation/AnimationWarping/Source/Runtime/Public/BoneControllers/AnimNode_OffsetRootBone.h
--- a/Engine/Plugins/Animation/AnimationWarping/Source/Runtime/Public/BoneControllers/AnimNode_OffsetRootBone.h
+++ b/Engine/Plugins/Animation/AnimationWarping/Source/Runtime/Public/BoneControllers/AnimNode_OffsetRootBone.h
@@ -72,6 +72,9 @@ struct ANIMATIONWARPINGRUNTIME_API FAnimNode_OffsetRootBone : public FAnimNode_B
 #if WITH_EDITORONLY_DATA
 	UPROPERTY(EditAnywhere, Category = Evaluation, meta=(FoldProperty))
 	EWarpingEvaluationMode EvaluationMode = EWarpingEvaluationMode::Graph;
+	
+	UPROPERTY(EditAnywhere, Category = Settings, meta = (FoldProperty, PinHiddenByDefault))
+	bool bResetEveryFrame = false;
 
 	// The translation offset behavior mode
 	UPROPERTY(EditAnywhere, Category = Settings, meta = (FoldProperty, PinHiddenByDefault))
@@ -152,6 +155,7 @@ public:
 
 	// Folded property accesors
 	EWarpingEvaluationMode GetEvaluationMode() const;
+	bool GetResetEveryFrame() const;
 	const FVector& GetTranslationDelta() const;
 	const FRotator& GetRotationDelta() const;
 	EOffsetRootBoneMode GetTranslationMode() const;
diff --git a/Engine/Plugins/Animation/AnimationWarping/Source/Runtime/Public/BoneControllers/AnimNode_OverrideRootMotion.h b/Engine/Plugins/Animation/AnimationWarping/Source/Runtime/Public/BoneControllers/AnimNode_OverrideRootMotion.h
new file mode 100644
--- /dev/null
+++ b/Engine/Plugins/Animation/AnimationWarping/Source/Runtime/Public/BoneControllers/AnimNode_OverrideRootMotion.h
@@ -0,0 +1,49 @@
+// Copyright Epic Games, Inc. All Rights Reserved.
+
+#pragma once
+
+#include "BoneControllers/BoneControllerTypes.h"
+#include "BoneControllers/AnimNode_SkeletalControlBase.h"
+#include "AnimNode_OverrideRootMotion.generated.h"
+
+struct FAnimationInitializeContext;
+struct FComponentSpacePoseContext;
+struct FNodeDebugData;
+
+USTRUCT(BlueprintInternalUseOnly, Experimental)
+struct ANIMATIONWARPINGRUNTIME_API FAnimNode_OverrideRootMotion : public FAnimNode_Base
+{
+	GENERATED_BODY();
+
+	UPROPERTY(EditAnywhere, BlueprintReadWrite, Category = Links, meta = (DisplayPriority = 0))
+	FPoseLink Source;
+
+#if WITH_EDITORONLY_DATA
+	UPROPERTY(EditAnywhere, Category = Evaluation, meta = (FoldProperty, PinShownByDefault))
+	float Alpha = 1.f;
+
+	UPROPERTY(EditAnywhere, Category = Evaluation, meta = (FoldProperty, PinShownByDefault))
+	FVector OverrideVelocity = FVector(0,0,0);
+
+	//todo: rotation override support
+#endif
+
+public:
+	// FAnimNode_Base interface
+	virtual void Update_AnyThread(const FAnimationUpdateContext& Context) override;
+	virtual void Evaluate_AnyThread(FPoseContext& Output) override;
+	virtual void Initialize_AnyThread(const FAnimationInitializeContext& Context) override;
+	virtual void CacheBones_AnyThread(const FAnimationCacheBonesContext& Context) override;
+	// End of FAnimNode_Base interface
+
+	// Folded property accesors
+	float GetAlpha() const;
+	const FVector& GetOverrideVelocity() const;
+
+private:
+
+	// Internal cached anim instance proxy
+	FAnimInstanceProxy* AnimInstanceProxy = nullptr;
+
+	float DeltaTime = 0.f;
+};
diff --git a/Engine/Plugins/Animation/MotionWarping/Source/MotionWarping/Private/AttributeBasedRootMotionComponent.cpp b/Engine/Plugins/Animation/MotionWarping/Source/MotionWarping/Private/AttributeBasedRootMotionComponent.cpp
new file mode 100644
--- /dev/null
+++ b/Engine/Plugins/Animation/MotionWarping/Source/MotionWarping/Private/AttributeBasedRootMotionComponent.cpp
@@ -0,0 +1,57 @@
+// Copyright Epic Games, Inc. All Rights Reserved.
+
+#include "AttributeBasedRootMotionComponent.h"
+
+#include "Animation/AnimInstance.h"
+#include "Animation/AnimationPoseData.h"
+#include "BonePose.h"
+#include "Animation/BuiltInAttributeTypes.h"
+#include "Components/SkeletalMeshComponent.h"
+#include "GameFramework/Character.h"
+#include "GameFramework/CharacterMovementComponent.h"
+
+#include UE_INLINE_GENERATED_CPP_BY_NAME(AttributeBasedRootMotionComponent)
+
+// UAttributeBasedRootMotionComponent
+///////////////////////////////////////////////////////////////////////
+
+UAttributeBasedRootMotionComponent::UAttributeBasedRootMotionComponent(const FObjectInitializer& ObjectInitializer)
+	: Super(ObjectInitializer)
+{
+	bWantsInitializeComponent = true;
+	PrimaryComponentTick.TickGroup = TG_PostUpdateWork;
+	PrimaryComponentTick.bCanEverTick = true;
+}
+
+void UAttributeBasedRootMotionComponent::InitializeComponent()
+{
+	Super::InitializeComponent();
+
+	CharacterOwner = Cast<ACharacter>(GetOwner());
+}
+
+void UAttributeBasedRootMotionComponent::TickComponent(float DeltaTime, enum ELevelTick TickType, FActorComponentTickFunction* ThisTickFunction)
+{
+	if (!bEnableRootMotion)
+	{
+		return;
+	}
+	
+	static const FName RootMotionAttributeName = "RootMotionDelta";
+	static const UE::Anim::FAttributeId RootMotionAttributeId = { RootMotionAttributeName , FCompactPoseBoneIndex(0) };
+	
+	Super::TickComponent(DeltaTime, TickType, ThisTickFunction);
+	
+	ACharacter* Character = GetCharacterOwner();
+	check(Character);
+
+	UCharacterMovementComponent* CharacterMovement = Character->GetCharacterMovement<UCharacterMovementComponent>();
+
+	if (const USkeletalMeshComponent* Mesh = Character->GetMesh())
+	{
+		if (const FTransformAnimationAttribute* RootMotionAttribute = Mesh->GetCustomAttributes().Find<FTransformAnimationAttribute>(RootMotionAttributeId))
+		{
+			CharacterMovement->RootMotionParams.Set(RootMotionAttribute->Value);
+		}
+	}
+}
diff --git a/Engine/Plugins/Animation/MotionWarping/Source/MotionWarping/Public/AttributeBasedRootMotionComponent.h b/Engine/Plugins/Animation/MotionWarping/Source/MotionWarping/Public/AttributeBasedRootMotionComponent.h
new file mode 100644
--- /dev/null
+++ b/Engine/Plugins/Animation/MotionWarping/Source/MotionWarping/Public/AttributeBasedRootMotionComponent.h
@@ -0,0 +1,37 @@
+// Copyright Epic Games, Inc. All Rights Reserved.
+
+#pragma once
+
+#include "CoreMinimal.h"
+#include "Components/ActorComponent.h"
+#include "AttributeBasedRootMotionComponent.generated.h"
+
+class ACharacter;
+
+UCLASS(ClassGroup = Movement, meta = (BlueprintSpawnableComponent))
+class MOTIONWARPING_API UAttributeBasedRootMotionComponent : public UActorComponent
+{
+	GENERATED_BODY()
+
+public:
+	UAttributeBasedRootMotionComponent(const FObjectInitializer& ObjectInitializer);
+
+	virtual void InitializeComponent() override;
+
+	virtual void TickComponent(float DeltaTime, enum ELevelTick TickType, FActorComponentTickFunction* ThisTickFunction) override;
+	
+	/** Gets the character this component belongs to */
+	FORCEINLINE ACharacter* GetCharacterOwner() const { return CharacterOwner.Get(); }
+
+	UPROPERTY(BlueprintReadWrite, Category = "Runtime")
+	bool bEnableRootMotion = true;
+	
+protected:
+
+	/** Character this component belongs to */
+	UPROPERTY(Transient)
+	TWeakObjectPtr<ACharacter> CharacterOwner;
+
+	FTransform PrevTransform;
+	bool PrevTransformValid = false;
+};
\ No newline at end of file
diff --git a/Engine/Plugins/Animation/PoseSearch/Source/Runtime/Private/PoseSearchTrajectoryLibrary.cpp b/Engine/Plugins/Animation/PoseSearch/Source/Runtime/Private/PoseSearchTrajectoryLibrary.cpp
--- a/Engine/Plugins/Animation/PoseSearch/Source/Runtime/Private/PoseSearchTrajectoryLibrary.cpp
+++ b/Engine/Plugins/Animation/PoseSearch/Source/Runtime/Private/PoseSearchTrajectoryLibrary.cpp
@@ -476,3 +476,10 @@ void UPoseSearchTrajectoryLibrary::GetTrajectorySampleAtTime(UPARAM(ref) const F
 	OutTrajectorySample = InTrajectory.GetSampleAtTime(Time, bExtrapolate);
 }
 
+void UPoseSearchTrajectoryLibrary::GetTrajectoryVelocity(UPARAM(ref) const FPoseSearchQueryTrajectory& InTrajectory, float Time1, float Time2, FVector& OutVelocity, bool bExtrapolate)
+{
+	FPoseSearchQueryTrajectorySample Sample1 = InTrajectory.GetSampleAtTime(Time1, bExtrapolate);
+	FPoseSearchQueryTrajectorySample Sample2 = InTrajectory.GetSampleAtTime(Time2, bExtrapolate);
+
+	OutVelocity = (Sample2.Position - Sample1.Position) / (Time2 - Time1);
+}
diff --git a/Engine/Plugins/Animation/PoseSearch/Source/Runtime/Public/PoseSearch/PoseSearchTrajectoryLibrary.h b/Engine/Plugins/Animation/PoseSearch/Source/Runtime/Public/PoseSearch/PoseSearchTrajectoryLibrary.h
--- a/Engine/Plugins/Animation/PoseSearch/Source/Runtime/Public/PoseSearch/PoseSearchTrajectoryLibrary.h
+++ b/Engine/Plugins/Animation/PoseSearch/Source/Runtime/Public/PoseSearch/PoseSearchTrajectoryLibrary.h
@@ -129,6 +129,15 @@ public:
 	UFUNCTION(BlueprintPure, meta = (BlueprintThreadSafe), Category="Animation|PoseSearch")
 	static void GetTrajectorySampleAtTime(UPARAM(ref) const FPoseSearchQueryTrajectory& InTrajectory, float Time, FPoseSearchQueryTrajectorySample& OutTrajectorySample, bool bExtrapolate = false);
 
+	UFUNCTION(BlueprintPure, meta = (BlueprintThreadSafe), Category="Animation|PoseSearch")
+	static void GetTrajectoryVelocity(UPARAM(ref) const FPoseSearchQueryTrajectory& InTrajectory, float Time1, float Time2, FVector& OutVelocity, bool bExtrapolate = false);
+	
+	UFUNCTION(BlueprintPure, meta = (BlueprintThreadSafe), Category="Animation|PoseSearch")
+	static FTransform GetTransform(UPARAM(ref) const FPoseSearchQueryTrajectorySample& InTrajectorySample)
+	{
+		return InTrajectorySample.GetTransform();
+	};
+	
 private:
 	static FVector RemapVectorMagnitudeWithCurve(const FVector& Vector, bool bUseCurve, const FRuntimeFloatCurve& Curve);
 };
\ No newline at end of file
