Improved comments on orientation warping anim node

diff --git a/Engine/Plugins/Animation/AnimationWarping/Source/Runtime/Public/BoneControllers/AnimNode_OrientationWarping.h b/Engine/Plugins/Animation/AnimationWarping/Source/Runtime/Public/BoneControllers/AnimNode_OrientationWarping.h
--- a/Engine/Plugins/Animation/AnimationWarping/Source/Runtime/Public/BoneControllers/AnimNode_OrientationWarping.h
+++ b/Engine/Plugins/Animation/AnimationWarping/Source/Runtime/Public/BoneControllers/AnimNode_OrientationWarping.h
@@ -37,11 +37,13 @@ struct ANIMATIONWARPINGRUNTIME_API FAnimNode_OrientationWarping : public FAnimNo
 
 	// The character locomotion angle (in degrees) relative to the specified RotationAxis
 	// This will be used in the following equation for computing the orientation angle: [Orientation = RotationBetween(RootMotionDirection, LocomotionDirection)]
+	// In most cases, this is the difference between the Velocity of the Movement Component and the actor rotation (obtained via CalculateDirection)
 	UPROPERTY(EditAnywhere, BlueprintReadWrite, Category=Evaluation, meta=(PinShownByDefault))
 	float LocomotionAngle = 0.f;
 
 	// The character movement direction vector in world space
-	// This will be used to compute LocomotionAngle automatically
+	// When set, this vector is used to compute LocomotionAngle automatically. When not set, the LocomotionAngle input should be used instead.
+	// In most cases, this vector is the same as the Velocity vector of the Movement Component
 	UPROPERTY(EditAnywhere, BlueprintReadWrite, Category=Evaluation, meta=(PinShownByDefault))
 	FVector LocomotionDirection = { 0.f, 0.f, 0.f };
 	
