Add option in FootPlacement anim node to disable complex collision traces.

diff --git a/Engine/Plugins/Animation/AnimationWarping/Source/Runtime/Private/BoneControllers/AnimNode_FootPlacement.cpp b/Engine/Plugins/Animation/AnimationWarping/Source/Runtime/Private/BoneControllers/AnimNode_FootPlacement.cpp
--- a/Engine/Plugins/Animation/AnimationWarping/Source/Runtime/Private/BoneControllers/AnimNode_FootPlacement.cpp
+++ b/Engine/Plugins/Animation/AnimationWarping/Source/Runtime/Private/BoneControllers/AnimNode_FootPlacement.cpp
@@ -504,8 +504,13 @@ void FAnimNode_FootPlacement::UpdatePlantingPlaneInterpolation(
 	const FVector TraceDirection = Context.ApproachDirWS;
 	const FPlane LastPlantPlane = InOutPlantPlane;
 
-	const bool bTraceAgainstSimpleAndComplex = TraceSettings.SimpleCollisionInfluence > 0.0f;
-	if (bTraceAgainstSimpleAndComplex &&!FMath::IsNearlyEqual(AlignmentAlpha, 1.0f))
+	if (TraceSettings.bDisableComplexTrace)
+	{
+		FVector ImpactLocationWS;
+		// Trace against simple geometry only
+		UE::Anim::FootPlacement::FindPlantPlane(Context, TraceSettings, FootTransformWS.GetLocation(), false, InOutPlantPlane, ImpactLocationWS);
+	}
+	else if (TraceSettings.SimpleCollisionInfluence > 0.0f && !FMath::IsNearlyEqual(AlignmentAlpha, 1.0f))
 	{
 		FVector ImpactLocationSimpleWS; 
 		FVector ImpactLocationComplexWS;
diff --git a/Engine/Plugins/Animation/AnimationWarping/Source/Runtime/Public/BoneControllers/AnimNode_FootPlacement.h b/Engine/Plugins/Animation/AnimationWarping/Source/Runtime/Public/BoneControllers/AnimNode_FootPlacement.h
--- a/Engine/Plugins/Animation/AnimationWarping/Source/Runtime/Public/BoneControllers/AnimNode_FootPlacement.h
+++ b/Engine/Plugins/Animation/AnimationWarping/Source/Runtime/Public/BoneControllers/AnimNode_FootPlacement.h
@@ -284,8 +284,11 @@ public:
 	UPROPERTY(EditAnywhere, Category = "Trace Settings", meta=(EditCondition="bEnabled", DisplayAfter="bEnabled"))
 	float SweepRadius = 5.0f;
 
-	// The channel to use for our complex trace
 	UPROPERTY(EditAnywhere, Category = "Trace Settings", meta=(EditCondition="bEnabled", DisplayAfter="bEnabled"))
+	bool bDisableComplexTrace = false;
+
+	// The channel to use for our complex trace
+	UPROPERTY(EditAnywhere, Category = "Trace Settings", meta=(EditCondition="!bDisableComplexTrace", DisplayAfter="bDisableComplexTrace"))
 	TEnumAsByte<ETraceTypeQuery> ComplexTraceChannel = TraceTypeQuery1;
 
 	// How much the feet can penetrate the ground geometry. It's recommended to allow some to account for interpolation
