Allow Root Motion Attribute to be set directly when it does not already exist.

diff --git a/Engine/Plugins/Animation/AnimationWarping/Source/Runtime/Private/AnimationWarpingRuntimeModule.cpp b/Engine/Plugins/Animation/AnimationWarping/Source/Runtime/Private/AnimationWarpingRuntimeModule.cpp
--- a/Engine/Plugins/Animation/AnimationWarping/Source/Runtime/Private/AnimationWarpingRuntimeModule.cpp
+++ b/Engine/Plugins/Animation/AnimationWarping/Source/Runtime/Private/AnimationWarpingRuntimeModule.cpp
@@ -25,6 +25,7 @@ public:
 
 public:
 	virtual void SampleRootMotion(const FDeltaTimeRecord& SampleRange, const UAnimSequence& Sequence, bool bLoopingSequence, UE::Anim::FStackAttributeContainer& OutAttributes) const override;
+	virtual bool SetRootMotion(const FTransform& RootMotionDelta, UE::Anim::FStackAttributeContainer& OutAttributes) const override;
 	virtual bool OverrideRootMotion(const FTransform& RootMotionDelta, UE::Anim::FStackAttributeContainer& OutAttributes) const override;
 	virtual bool ExtractRootMotion(const UE::Anim::FStackAttributeContainer& Attributes, FTransform& OutRootMotionDelta) const override;
 	virtual bool HasRootMotion(const UE::Anim::FStackAttributeContainer& Attributes) const override;
@@ -55,6 +56,17 @@ void FModule::SampleRootMotion(const FDeltaTimeRecord& SampleRange, const UAnimS
 	RootMotionAttribute->Value = RootMotionTransform;
 }
 
+bool FModule::SetRootMotion(const FTransform& RootMotionDelta, UE::Anim::FStackAttributeContainer& OutAttributes) const
+{
+	if (FTransformAnimationAttribute* RootMotionAttribute = OutAttributes.FindOrAdd<FTransformAnimationAttribute>(RootMotionAttributeId))
+	{
+		RootMotionAttribute->Value = RootMotionDelta;
+		return true;
+	}
+
+	return false;
+}
+
 bool FModule::OverrideRootMotion(const FTransform& RootMotionDelta, UE::Anim::FStackAttributeContainer& OutAttributes) const
 {
 	if (FTransformAnimationAttribute* RootMotionAttribute = OutAttributes.Find<FTransformAnimationAttribute>(RootMotionAttributeId))
diff --git a/Engine/Source/Runtime/Engine/Public/Animation/AnimRootMotionProvider.h b/Engine/Source/Runtime/Engine/Public/Animation/AnimRootMotionProvider.h
--- a/Engine/Source/Runtime/Engine/Public/Animation/AnimRootMotionProvider.h
+++ b/Engine/Source/Runtime/Engine/Public/Animation/AnimRootMotionProvider.h
@@ -30,6 +30,9 @@ public:
 	// Given the specified time range, sequence, and looping behavior, sample a root motion delta and store it in the specified attribute container
 	virtual void SampleRootMotion(const FDeltaTimeRecord& SampleRange, const UAnimSequence& Sequence, bool bLoopingSequence, FStackAttributeContainer& OutAttributes) const = 0;
 	
+	// Directly sets the root motion attribute. Does not require root motion to exist already.
+	virtual bool SetRootMotion(const FTransform& RootMotionDelta, FStackAttributeContainer& OutAttributes) const { return false; }
+
 	// Override the currently stored root motion delta in the specified attribute container. Requires root motion to exist prior to overriding.
 	virtual bool OverrideRootMotion(const FTransform& RootMotionDelta, FStackAttributeContainer& OutAttributes) const = 0;
 	
