Fix to prevent Offset Root Bone anim node overwiting scale

diff --git a/Engine/Plugins/Animation/AnimationWarping/Source/Runtime/Private/BoneControllers/AnimNode_OffsetRootBone.cpp b/Engine/Plugins/Animation/AnimationWarping/Source/Runtime/Private/BoneControllers/AnimNode_OffsetRootBone.cpp
--- a/Engine/Plugins/Animation/AnimationWarping/Source/Runtime/Private/BoneControllers/AnimNode_OffsetRootBone.cpp
+++ b/Engine/Plugins/Animation/AnimationWarping/Source/Runtime/Private/BoneControllers/AnimNode_OffsetRootBone.cpp
@@ -333,6 +333,9 @@ void FAnimNode_OffsetRootBone::Evaluate_AnyThread(FPoseContext& Output)
 	// Accumulate the input bone transform to keep the offset independent from any previous adjustments to the root
 	TargetBoneTransform.Accumulate(InputBoneTransform);
 
+	// Offset root bone should not affect scale so take the input
+	TargetBoneTransform.SetScale3D(InputBoneTransform.GetScale3D());
+
 	Output.Pose[TargetBoneIndex] = TargetBoneTransform;
 
 #if ENABLE_VISUAL_LOG
