Added flag to smooth out root bone using foot placement

diff --git a/Engine/Plugins/Animation/AnimationWarping/Source/Runtime/Private/BoneControllers/AnimNode_FootPlacement.cpp b/Engine/Plugins/Animation/AnimationWarping/Source/Runtime/Private/BoneControllers/AnimNode_FootPlacement.cpp
--- a/Engine/Plugins/Animation/AnimationWarping/Source/Runtime/Private/BoneControllers/AnimNode_FootPlacement.cpp
+++ b/Engine/Plugins/Animation/AnimationWarping/Source/Runtime/Private/BoneControllers/AnimNode_FootPlacement.cpp
@@ -1230,6 +1230,17 @@ void FAnimNode_FootPlacement::EvaluateSkeletalControl_AnyThread(FComponentSpaceP
 	check(!PelvisTransformCS.ContainsNaN());
 	OutBoneTransforms.Add(FBoneTransform(PelvisData.Bones.FkBoneIndex, PelvisTransformCS));
 
+	if (InterpolationSettings.bSmoothRootBone)
+	{
+		// Smooth out the root by the same factor as the hips.
+		FTransform RootBoneTransform = Output.Pose.GetComponentSpaceTransform(FCompactPoseBoneIndex(0));
+		const FVector TranslationDelta = PelvisTransformCS.GetTranslation() - PelvisData.InputPose.FKTransformCS.GetTranslation();
+
+		RootBoneTransform.AddToTranslation(TranslationDelta);
+		check(!RootBoneTransform.ContainsNaN());
+		OutBoneTransforms.Add(FBoneTransform(FCompactPoseBoneIndex(0), RootBoneTransform));
+	}
+
 	for (int32 FootIndex = 0; FootIndex < LegsData.Num(); ++FootIndex)
 	{
 		UE::Anim::FootPlacement::FLegRuntimeData& LegData = LegsData[FootIndex];
diff --git a/Engine/Plugins/Animation/AnimationWarping/Source/Runtime/Public/BoneControllers/AnimNode_FootPlacement.h b/Engine/Plugins/Animation/AnimationWarping/Source/Runtime/Public/BoneControllers/AnimNode_FootPlacement.h
--- a/Engine/Plugins/Animation/AnimationWarping/Source/Runtime/Public/BoneControllers/AnimNode_FootPlacement.h
+++ b/Engine/Plugins/Animation/AnimationWarping/Source/Runtime/Public/BoneControllers/AnimNode_FootPlacement.h
@@ -260,6 +260,9 @@ public:
 	UPROPERTY(EditAnywhere, Category = "Plant Settings")
 	bool bEnableFloorInterpolation = true;
 	
+	UPROPERTY(EditAnywhere, Category = "Plant Settings")
+	bool bSmoothRootBone = false;
+
 	UPROPERTY(EditAnywhere, Category = "Plant Settings")
 	bool bEnableSeparationInterpolation = true;
 };
