OffsetRootBone: Fix order of operations when input has a transform

diff --git a/Engine/Plugins/Animation/AnimationWarping/Source/Runtime/Private/BoneControllers/AnimNode_OffsetRootBone.cpp b/Engine/Plugins/Animation/AnimationWarping/Source/Runtime/Private/BoneControllers/AnimNode_OffsetRootBone.cpp
--- a/Engine/Plugins/Animation/AnimationWarping/Source/Runtime/Private/BoneControllers/AnimNode_OffsetRootBone.cpp
+++ b/Engine/Plugins/Animation/AnimationWarping/Source/Runtime/Private/BoneControllers/AnimNode_OffsetRootBone.cpp
@@ -328,10 +328,10 @@ void FAnimNode_OffsetRootBone::Evaluate_AnyThread(FPoseContext& Output)
 	SimulatedTransform.SetLocation(SimulatedTranslation);
 	SimulatedTransform.SetRotation(SimulatedRotation);
 
-	// Combine with the input pose's bone transform, to preserve any adjustments done before this node in the graph
-	FTransform TargetBoneTransform = SimulatedTransform * ComponentTransform.Inverse();
-	// Accumulate the input bone transform to keep the offset independent from any previous adjustments to the root
-	TargetBoneTransform.Accumulate(InputBoneTransform);
+	// Start with the input pose's bone transform, to preserve any adjustments done before this node in the graph
+	FTransform TargetBoneTransform = InputBoneTransform;
+	// Accumulate the simulated transform in, and counter current component transform.
+	TargetBoneTransform.Accumulate(SimulatedTransform * ComponentTransform.Inverse());
 
 	// Offset root bone should not affect scale so take the input
 	TargetBoneTransform.SetScale3D(InputBoneTransform.GetScale3D());
