Add prediction to orientation warping
Fix crash preview root motion

#tests editor, PIE

diff --git a/Engine/Plugins/Animation/AnimationWarping/Source/Runtime/Private/BoneControllers/AnimNode_OrientationWarping.cpp b/Engine/Plugins/Animation/AnimationWarping/Source/Runtime/Private/BoneControllers/AnimNode_OrientationWarping.cpp
--- a/Engine/Plugins/Animation/AnimationWarping/Source/Runtime/Private/BoneControllers/AnimNode_OrientationWarping.cpp
+++ b/Engine/Plugins/Animation/AnimationWarping/Source/Runtime/Private/BoneControllers/AnimNode_OrientationWarping.cpp
@@ -151,6 +151,8 @@ void FAnimNode_OrientationWarping::EvaluateSkeletalControl_AnyThread(FComponentS
 	FTransform RootMotionTransformDelta = FTransform::Identity;
 	float RootMotionDeltaAngleRad = 0.0;
 	const float PreviousOrientationAngleRad = ActualOrientationAngleRad;
+	FVector CurrentRootMotionDeltaDirection = FVector::ZeroVector;
+	bool bUsedFutureRootMotion = false;
 #endif
 
 	// We will likely need to revisit LocomotionAngle participating as an input to orientation warping.
@@ -227,6 +229,18 @@ void FAnimNode_OrientationWarping::EvaluateSkeletalControl_AnyThread(FComponentS
 
 			// Flatten root motion translation, along the rotation axis.
 			RootMotionDeltaTranslation = RootMotionDeltaTranslation - RotationAxisVector.Dot(RootMotionDeltaTranslation) * RotationAxisVector;
+
+			// Look forward in root motion to determine orientation warping angle if future time specified
+			FVector PredictedRootMotionDeltaTranslation = RootMotionDeltaTranslation;
+			if (TargetTime > 0 && CurrentAnimAsset)
+			{
+				if (UAnimSequenceBase* AnimSequence = Cast<UAnimSequenceBase>(CurrentAnimAsset))
+				{
+					FTransform PredictedRootMotionTransform = AnimSequence->ExtractRootMotion(CurrentAnimAssetTime, TargetTime, false);
+					PredictedRootMotionDeltaTranslation = PredictedRootMotionTransform.GetTranslation();
+					PredictedRootMotionDeltaTranslation = PredictedRootMotionDeltaTranslation - RotationAxisVector.Dot(PredictedRootMotionDeltaTranslation) * RotationAxisVector;
+				}
+			}
 			
 			const float RootMotionDeltaSpeed = RootMotionDeltaTranslation.Size() / DeltaSeconds;
 			if (RootMotionDeltaSpeed < MinRootMotionSpeedThreshold)
@@ -236,7 +250,6 @@ void FAnimNode_OrientationWarping::EvaluateSkeletalControl_AnyThread(FComponentS
 			}
 			else
 			{
-
 				const FVector PreviousRootMotionDeltaDirection = RootMotionDeltaDirection;
 				// Hold previous direction if we can't calculate it from current move delta, because the root is no longer moving
 				RootMotionDeltaDirection = RootMotionDeltaTranslation.GetSafeNormal(UE_SMALL_NUMBER, PreviousRootMotionDeltaDirection);
@@ -246,6 +259,7 @@ void FAnimNode_OrientationWarping::EvaluateSkeletalControl_AnyThread(FComponentS
 				// When that happens, since we use the delta between root motion and movement direction, we would be over-rotating the lower body and breaking the pose during those frames
 				// So, when that happens we use the inverse of the root motion direction to calculate our target rotation. 
 				// This feels a bit 'hacky' but its the only option I've found so far to mitigate the problem
+				float PreFlipOrientationAngleRad = TargetOrientationAngleRad;
 				if (LocomotionAngleDeltaThreshold > 0.f)
 				{
 					if (FMath::Abs(FMath::RadiansToDegrees(TargetOrientationAngleRad)) > LocomotionAngleDeltaThreshold)
@@ -255,6 +269,33 @@ void FAnimNode_OrientationWarping::EvaluateSkeletalControl_AnyThread(FComponentS
 					}
 				}
 
+				// If there is translation in predicted root motion, use it if the orientation error is less than current error
+				if (TargetTime > 0 && CurrentAnimAsset && !PredictedRootMotionDeltaTranslation.IsNearlyZero(UE_SMALL_NUMBER))
+				{
+					PredictedRootMotionDeltaTranslation.Normalize();
+					float PredictedOrientationErrorAngleRad = UE::Anim::SignedAngleRadBetweenNormals(PredictedRootMotionDeltaTranslation, LocomotionForward, RotationAxisVector);
+
+					// The future orientation will often match the current, so add a small delta to avoid further testing for same values
+					if (FMath::Abs(PredictedOrientationErrorAngleRad) + UE_KINDA_SMALL_NUMBER < FMath::Abs(PreFlipOrientationAngleRad))
+					{
+						// Error is normalized. So even if the future root motion orientation has less error, it may be in the opposite angle.
+						// Thus only use the future root motion for warping if it also has less error relative to the previous position
+						float PredictedToPrevErrorAngleRad = UE::Anim::SignedAngleRadBetweenNormals(PreviousRootMotionDeltaDirection, PredictedRootMotionDeltaTranslation, RotationAxisVector);
+						float CurrentToPrevErrorAngleRad = UE::Anim::SignedAngleRadBetweenNormals(PreviousRootMotionDeltaDirection, RootMotionDeltaDirection, RotationAxisVector);
+
+						// Check for continuity. Later, when the future and current are re-aligned, error will be equivalent, so we can fall back to non-future root motion
+						if (FMath::Abs(PredictedToPrevErrorAngleRad) < FMath::Abs(CurrentToPrevErrorAngleRad) + FMath::DegreesToRadians(TargetBiasAngle))
+						{
+#if ENABLE_ANIM_DEBUG
+							CurrentRootMotionDeltaDirection = RootMotionDeltaDirection;
+							bUsedFutureRootMotion = true;
+#endif
+							TargetOrientationAngleRad = PredictedOrientationErrorAngleRad;
+							RootMotionDeltaDirection = PredictedRootMotionDeltaTranslation;
+						}
+					}
+				}
+
 				// Don't compensate interpolation by the root motion angle delta if the previous angle isn't valid.
 				if (bCounterCompenstateInterpolationByRootMotion && !PreviousRootMotionDeltaDirection.IsNearlyZero(UE_SMALL_NUMBER))
 				{
@@ -351,10 +392,28 @@ void FAnimNode_OrientationWarping::EvaluateSkeletalControl_AnyThread(FComponentS
 				: ActorForwardDirection.RotateAngleAxis(OrientationAngle, RotationAxisVector);
 
 			DebugArrowOffset += FVector::ZAxisVector * DebugDrawScale;
-			Output.AnimInstanceProxy->AnimDrawDebugDirectionalArrow(
-				ComponentTransform.GetLocation() + DebugArrowOffset,
-				ComponentTransform.GetLocation() + DebugArrowOffset + RotationDirection * 100.f * DebugDrawScale,
-				40.f * DebugDrawScale, FColor::Blue, false, 0.f, 2.f * DebugDrawScale);
+
+			if (bUsedFutureRootMotion && bGraphDrivenWarping)
+			{
+				const FVector CurrentRotationDirection = ComponentTransform.GetRotation().RotateVector(CurrentRootMotionDeltaDirection);
+
+				Output.AnimInstanceProxy->AnimDrawDebugDirectionalArrow(
+					ComponentTransform.GetLocation() + DebugArrowOffset,
+					ComponentTransform.GetLocation() + DebugArrowOffset + CurrentRotationDirection * 100.f * DebugDrawScale,
+					40.f * DebugDrawScale, FColor::Blue, false, 0.f, 2.f * DebugDrawScale);
+
+				Output.AnimInstanceProxy->AnimDrawDebugDirectionalArrow(
+					ComponentTransform.GetLocation() + DebugArrowOffset,
+					ComponentTransform.GetLocation() + DebugArrowOffset + RotationDirection * 100.f * DebugDrawScale,
+					40.f * DebugDrawScale, FColor::Cyan, false, 0.f, 2.f * DebugDrawScale);
+			}
+			else
+			{
+				Output.AnimInstanceProxy->AnimDrawDebugDirectionalArrow(
+					ComponentTransform.GetLocation() + DebugArrowOffset,
+					ComponentTransform.GetLocation() + DebugArrowOffset + RotationDirection * 100.f * DebugDrawScale,
+					40.f * DebugDrawScale, FColor::Blue, false, 0.f, 2.f * DebugDrawScale);
+			}
 
 			const float ActualOrientationAngleDegrees = FMath::RadiansToDegrees(ActualOrientationAngleRad);
 			const FVector WarpedRotationDirection = bGraphDrivenWarping
diff --git a/Engine/Plugins/Animation/AnimationWarping/Source/Runtime/Public/BoneControllers/AnimNode_OrientationWarping.h b/Engine/Plugins/Animation/AnimationWarping/Source/Runtime/Public/BoneControllers/AnimNode_OrientationWarping.h
--- a/Engine/Plugins/Animation/AnimationWarping/Source/Runtime/Public/BoneControllers/AnimNode_OrientationWarping.h
+++ b/Engine/Plugins/Animation/AnimationWarping/Source/Runtime/Public/BoneControllers/AnimNode_OrientationWarping.h
@@ -22,6 +22,10 @@ enum class EOrientationWarpingSpace : uint8
 	CustomTransform
 };
 
+/**
+ * Maintains a look at direction for the upper body (orientation), while rotating the lower body to match capsule velocity direction
+ * Does nothing if the root motion velocity direction matches the desired / current capsule velocity direction
+ */
 USTRUCT(BlueprintInternalUseOnly)
 struct ANIMATIONWARPINGRUNTIME_API FAnimNode_OrientationWarping : public FAnimNode_SkeletalControlBase
 {
@@ -31,6 +35,19 @@ struct ANIMATIONWARPINGRUNTIME_API FAnimNode_OrientationWarping : public FAnimNo
 	UPROPERTY(EditAnywhere, Category=Evaluation)
 	EWarpingEvaluationMode Mode = EWarpingEvaluationMode::Manual;
 
+	// Experimental. Orientation warping should do nothing if root motion velocity directions match capsule,
+	// however root motion can have multiple velocity directions. So we also check root motion
+	// direction 'TargetTime' in the future for matching direction to avoid temp orientation warps.
+	UPROPERTY(EditAnywhere, BlueprintReadWrite, Category = Experimental, meta = (PinShownByDefault))
+	float TargetTime = 0.8f;
+
+	// Experimental. A bias in which the future root motion is preferred, even if the current root motion is closer to the previous orientation.
+	// We compare the current / future root motion to our previous orientation for continuity. 
+	// By default the current root motion is used, so it will always win this comparision without some bias.
+	// Note: We default to 5 degrees which may seem large, but the animation itself being continuous should prevent popping.
+	UPROPERTY(EditAnywhere, BlueprintReadWrite, Category = Experimental, meta = (PinHiddenByDefault))
+	float TargetBiasAngle = 5.f;
+
 	// The desired orientation angle (in degrees) to warp by relative to the specified RotationAxis
 	UPROPERTY(EditAnywhere, BlueprintReadWrite, Category=Evaluation, meta=(PinShownByDefault))
 	float OrientationAngle = 0.f;
@@ -77,6 +94,15 @@ struct ANIMATIONWARPINGRUNTIME_API FAnimNode_OrientationWarping : public FAnimNo
 	UPROPERTY(EditAnywhere, Category=Settings, meta=(DisplayName="IK Foot Bones"))
 	TArray<FBoneReference> IKFootBones;
 
+	// Experimental. Animation Asset for incorporating root motion data. If 'TargetTime' is set, and the animation has root motion rotation within the TargetTime, 
+	// then those rotations will be scaled to reach the TargetOrientation
+	UPROPERTY(EditAnywhere, Transient, BlueprintReadWrite, Category = Experimental, meta = (PinShownByDefault))
+	TObjectPtr<UAnimationAsset> CurrentAnimAsset;
+
+	// Experimental. Current playback time in seconds of the CurrentAnimAsset
+	UPROPERTY(EditAnywhere, Transient, BlueprintReadWrite, Category = Experimental, meta = (PinShownByDefault))
+	float CurrentAnimAssetTime = 0.f;
+
 	// Rotation axis used when rotating the character body
 	UPROPERTY(EditAnywhere, Category=Settings)
 	TEnumAsByte<EAxis::Type> RotationAxis = EAxis::Z;
diff --git a/Engine/Source/Runtime/Engine/Private/Animation/AnimSequenceHelpers.cpp b/Engine/Source/Runtime/Engine/Private/Animation/AnimSequenceHelpers.cpp
--- a/Engine/Source/Runtime/Engine/Private/Animation/AnimSequenceHelpers.cpp
+++ b/Engine/Source/Runtime/Engine/Private/Animation/AnimSequenceHelpers.cpp
@@ -240,12 +240,15 @@ FTransform ExtractRootTransformFromAnimationAsset(const UAnimationAsset* Animati
 	bool bUseNormalizedRootMotionScale = false;
 
 	FTransform RootTransformRefPose = FTransform::Identity;
-	if (const USkeleton* MySkeleton = Animation->GetSkeleton())
+	if (Animation)
 	{
-		const FReferenceSkeleton& RefSkeleton = MySkeleton->GetReferenceSkeleton();
-		if (RefSkeleton.GetNum() > 0)
+		if (const USkeleton* MySkeleton = Animation->GetSkeleton())
 		{
-			RootTransformRefPose = RefSkeleton.GetRefBonePose()[0];
+			const FReferenceSkeleton& RefSkeleton = MySkeleton->GetReferenceSkeleton();
+			if (RefSkeleton.GetNum() > 0)
+			{
+				RootTransformRefPose = RefSkeleton.GetRefBonePose()[0];
+			}
 		}
 	}
 
@@ -274,7 +277,7 @@ FTransform ExtractRootTransformFromAnimationAsset(const UAnimationAsset* Animati
 		//Clear scale as it will muck up GetRelativeTransform
 		Result.SetScale3D(FVector(1.f));
 	}
-	else
+	else if (Animation)
 	{
 		if (Animation->IsValidAdditive())
 		{
