Offset Root Bone: Added options to better support over-rotation during turn in place

diff --git a/Engine/Plugins/Animation/AnimationWarping/Source/Runtime/Private/BoneControllers/AnimNode_OffsetRootBone.cpp b/Engine/Plugins/Animation/AnimationWarping/Source/Runtime/Private/BoneControllers/AnimNode_OffsetRootBone.cpp
--- a/Engine/Plugins/Animation/AnimationWarping/Source/Runtime/Private/BoneControllers/AnimNode_OffsetRootBone.cpp
+++ b/Engine/Plugins/Animation/AnimationWarping/Source/Runtime/Private/BoneControllers/AnimNode_OffsetRootBone.cpp
@@ -30,6 +30,42 @@ namespace UE::Anim::OffsetRootBone
 		Halflife = FMath::Max(Halflife, DeltaTime);
 		return FMath::Clamp((1.0f - FMath::InvExpApprox((0.69314718056f * DeltaTime) / (Halflife + Epsilon))), 0.0f, 1.0f);
 	}
+
+	static bool ShouldExtractRootMotion(const EOffsetRootBoneMode OffsetMode)
+	{
+		switch (OffsetMode)
+		{
+		case EOffsetRootBoneMode::Accumulate:
+		case EOffsetRootBoneMode::Interpolate:
+		case EOffsetRootBoneMode::LockOffsetAndConsumeAnimation:
+		case EOffsetRootBoneMode::LockOffsetIncreaseAndConsumeAnimation:
+			return true;
+			break;
+		case EOffsetRootBoneMode::LockOffsetAndIgnoreAnimation:
+		case EOffsetRootBoneMode::Release:
+		default:
+			return false;
+			break;
+		}
+	}
+
+	static bool ShouldCounterComponentDelta(const EOffsetRootBoneMode OffsetMode)
+	{
+		switch (OffsetMode)
+		{
+		case EOffsetRootBoneMode::Accumulate:
+		case EOffsetRootBoneMode::Interpolate:
+			return false;
+			break;
+		case EOffsetRootBoneMode::LockOffsetAndConsumeAnimation:
+		case EOffsetRootBoneMode::LockOffsetIncreaseAndConsumeAnimation:
+		case EOffsetRootBoneMode::LockOffsetAndIgnoreAnimation:
+		case EOffsetRootBoneMode::Release:
+		default:
+			return true;
+			break;
+		}
+	}
 }
 
 void FAnimNode_OffsetRootBone::GatherDebugData(FNodeDebugData& DebugData)
@@ -112,10 +148,11 @@ void FAnimNode_OffsetRootBone::Evaluate_AnyThread(FPoseContext& Output)
 	const FTransform LastComponentTransform = ComponentTransform;
 	ComponentTransform = AnimInstanceProxy->GetComponentTransform();
 
-	const bool bShouldConsumeTranslationOffset = GetTranslationMode() == EOffsetRootBoneMode::Accumulate ||
-										GetTranslationMode() == EOffsetRootBoneMode::Interpolate;
-	const bool bShouldConsumeRotationOffset = GetRotationMode() == EOffsetRootBoneMode::Accumulate ||
-										GetRotationMode() == EOffsetRootBoneMode::Interpolate;
+	const EOffsetRootBoneMode CurrentTranslationMode = GetTranslationMode();
+	const EOffsetRootBoneMode CurrentRotationMode = GetRotationMode();
+
+	bool bShouldConsumeTranslationOffset = UE::Anim::OffsetRootBone::ShouldExtractRootMotion(CurrentTranslationMode);
+	bool bShouldConsumeRotationOffset = UE::Anim::OffsetRootBone::ShouldExtractRootMotion(CurrentRotationMode);
 
 	FTransform RootMotionTransformDelta = FTransform::Identity;
 	if (bGraphDriven)
@@ -184,24 +221,54 @@ void FAnimNode_OffsetRootBone::Evaluate_AnyThread(FPoseContext& Output)
 		// Grab root motion translation from the root motion attribute
 		ConsumedRootMotionDelta.SetTranslation(RootMotionTransformDelta.GetTranslation());
 	}
-	else
-	{
-		// Accumulate the translation frame delta into the simulated translation, to keep it in place
-		const FVector ComponentTranslationDelta = ComponentTransform.GetLocation() - LastComponentTransform.GetLocation();
-		SimulatedTranslation += ComponentTranslationDelta;
-	}
-
 	if (bShouldConsumeRotationOffset)
 	{
 		// Grab root motion rotation from the root motion attribute
 		ConsumedRootMotionDelta.SetRotation(RootMotionTransformDelta.GetRotation());
 	}
-	else
+
+	if (UE::Anim::OffsetRootBone::ShouldCounterComponentDelta(CurrentRotationMode))
 	{
-		// Accumulate the rotation frame delta into the simulated translation, to keep it in place
+		// Accumulate the rotation component delta into the simulated rotation, to keep component and offset in sync.
 		const FQuat ComponentRotationDelta = LastComponentTransform.GetRotation().Inverse() * ComponentTransform.GetRotation();
 		SimulatedRotation = ComponentRotationDelta * SimulatedRotation;
 	}
+	if (UE::Anim::OffsetRootBone::ShouldCounterComponentDelta(CurrentTranslationMode))
+	{
+		// Accumulate the translation component delta into the simulated translation, to keep component and offset in sync.
+		const FVector ComponentTranslationDelta = ComponentTransform.GetLocation() - LastComponentTransform.GetLocation();
+		SimulatedTranslation += ComponentTranslationDelta;
+	}
+
+	if (CurrentTranslationMode == EOffsetRootBoneMode::LockOffsetIncreaseAndConsumeAnimation)
+	{
+		FVector DeltaDir = (SimulatedTranslation - ComponentTransform.GetLocation()).GetSafeNormal();
+		DeltaDir = SimulatedRotation.UnrotateVector(DeltaDir).GetSafeNormal();
+		// Only allow root motion to steer us towards an position that will make the offset smaller.
+		ConsumedRootMotionDelta.SetTranslation(DeltaDir * ConsumedRootMotionDelta.GetTranslation().Dot(DeltaDir));
+	}
+
+	if (CurrentRotationMode == EOffsetRootBoneMode::LockOffsetIncreaseAndConsumeAnimation)
+	{
+		FQuat DeltaRot = ComponentTransform.GetRotation() * SimulatedRotation.Inverse();
+
+		FVector DeltaAxis;
+		float DeltaAngle;
+		DeltaRot.ToAxisAndAngle(DeltaAxis, DeltaAngle);
+
+		float RootMotionAngle = ConsumedRootMotionDelta.GetRotation().GetTwistAngle(DeltaAxis);
+		if (DeltaAngle >= 0.0f)
+		{
+			RootMotionAngle = FMath::Clamp(RootMotionAngle, 0.0f, DeltaAngle);
+		}
+		else
+		{
+			RootMotionAngle = FMath::Clamp(RootMotionAngle, DeltaAngle, 0.0f);
+		}
+
+		// Only allow root motion to steer us towards an orientation that will make the offset smaller.
+		ConsumedRootMotionDelta.SetRotation(FQuat(DeltaAxis, RootMotionAngle));
+	}
 
 	FTransform SimulatedTransform(SimulatedRotation, SimulatedTranslation);
 	// Apply the root motion delta
diff --git a/Engine/Plugins/Animation/AnimationWarping/Source/Runtime/Public/BoneControllers/AnimNode_OffsetRootBone.h b/Engine/Plugins/Animation/AnimationWarping/Source/Runtime/Public/BoneControllers/AnimNode_OffsetRootBone.h
--- a/Engine/Plugins/Animation/AnimationWarping/Source/Runtime/Public/BoneControllers/AnimNode_OffsetRootBone.h
+++ b/Engine/Plugins/Animation/AnimationWarping/Source/Runtime/Public/BoneControllers/AnimNode_OffsetRootBone.h
@@ -42,9 +42,14 @@ enum class EOffsetRootBoneMode : uint8
 	// The root will stay behind, but will attempt to catch up
 	Interpolate,
 	// Stops accumulating the mesh component's movement delta into the root offset
-	// In this mode, whatever offset we have will be conserved, 
-	// but we won't accumulate any more
-	Hold,
+	// In this mode, whatever offset we have will be locked but we will still consume animated root motion 
+	LockOffsetAndConsumeAnimation,
+	// Stops accumulating the mesh component's movement delta into the root offset
+	// In this mode, whatever offset we have will be locked but we will still consume animated root motion, as long as it's decreasing the offset.
+	LockOffsetIncreaseAndConsumeAnimation,
+	// Stops accumulating the mesh component's movement delta into the root offset
+	// In this mode, whatever offset we have will be locked and we will ignore animated root motion
+	LockOffsetAndIgnoreAnimation,
 	// Release the offset and stop accumulating the mesh component's movement delta.
 	// In this mode we will "blend out" the offset
 	Release,
