AnimationWarping - keeping numerical errors in check in FAnimNode_OffsetRootBone::Evaluate_AnyThread

diff --git a/Engine/Plugins/Animation/AnimationWarping/Source/Runtime/Private/AnimationWarpingRuntimeModule.cpp b/Engine/Plugins/Animation/AnimationWarping/Source/Runtime/Private/AnimationWarpingRuntimeModule.cpp
--- a/Engine/Plugins/Animation/AnimationWarping/Source/Runtime/Private/AnimationWarpingRuntimeModule.cpp
+++ b/Engine/Plugins/Animation/AnimationWarping/Source/Runtime/Private/AnimationWarpingRuntimeModule.cpp
@@ -58,6 +58,7 @@ void FModule::SampleRootMotion(const FDeltaTimeRecord& SampleRange, const UAnimS
 
 bool FModule::SetRootMotion(const FTransform& RootMotionDelta, UE::Anim::FStackAttributeContainer& OutAttributes) const
 {
+	check(RootMotionDelta.IsRotationNormalized());
 	if (FTransformAnimationAttribute* RootMotionAttribute = OutAttributes.FindOrAdd<FTransformAnimationAttribute>(RootMotionAttributeId))
 	{
 		RootMotionAttribute->Value = RootMotionDelta;
@@ -69,6 +70,7 @@ bool FModule::SetRootMotion(const FTransform& RootMotionDelta, UE::Anim::FStackA
 
 bool FModule::OverrideRootMotion(const FTransform& RootMotionDelta, UE::Anim::FStackAttributeContainer& OutAttributes) const
 {
+	check(RootMotionDelta.IsRotationNormalized());
 	if (FTransformAnimationAttribute* RootMotionAttribute = OutAttributes.Find<FTransformAnimationAttribute>(RootMotionAttributeId))
 	{
 		RootMotionAttribute->Value = RootMotionDelta;
diff --git a/Engine/Plugins/Animation/AnimationWarping/Source/Runtime/Private/BoneControllers/AnimNode_OffsetRootBone.cpp b/Engine/Plugins/Animation/AnimationWarping/Source/Runtime/Private/BoneControllers/AnimNode_OffsetRootBone.cpp
--- a/Engine/Plugins/Animation/AnimationWarping/Source/Runtime/Private/BoneControllers/AnimNode_OffsetRootBone.cpp
+++ b/Engine/Plugins/Animation/AnimationWarping/Source/Runtime/Private/BoneControllers/AnimNode_OffsetRootBone.cpp
@@ -388,6 +388,7 @@ void FAnimNode_OffsetRootBone::Evaluate_AnyThread(FPoseContext& Output)
 		{
 			RotationOffset = FQuat(OffsetAxis, MaxAngleRadians);
 			SimulatedRotation = RotationOffset * ComponentTransform.GetRotation();
+			SimulatedRotation.Normalize();
 		}
 	}
 
@@ -403,6 +404,8 @@ void FAnimNode_OffsetRootBone::Evaluate_AnyThread(FPoseContext& Output)
 	// Offset root bone should not affect scale so take the input
 	TargetBoneTransform.SetScale3D(InputBoneTransform.GetScale3D());
 
+	TargetBoneTransform.NormalizeRotation();
+
 	Output.Pose[TargetBoneIndex] = TargetBoneTransform;
 
 #if ENABLE_VISUAL_LOG
diff --git a/Engine/Plugins/Animation/PoseSearch/Source/Runtime/Private/PoseSearchHistory.cpp b/Engine/Plugins/Animation/PoseSearch/Source/Runtime/Private/PoseSearchHistory.cpp
--- a/Engine/Plugins/Animation/PoseSearch/Source/Runtime/Private/PoseSearchHistory.cpp
+++ b/Engine/Plugins/Animation/PoseSearch/Source/Runtime/Private/PoseSearchHistory.cpp
@@ -287,6 +287,7 @@ int32 FPoseHistoryEntry::Num() const
 
 void FPoseHistoryEntry::SetComponentSpaceTransform(int32 Index, const FTransform& Transform)
 {
+	check( Transform.IsRotationNormalized() );
 	ComponentSpaceRotations[Index] = FQuat4f(Transform.GetRotation());
 	ComponentSpacePositions[Index] = Transform.GetTranslation();
 	
