Foot Placement: Improvements on slopes + changing component space data to be in root-space to work w/ offset root bone.

diff --git a/Engine/Plugins/Animation/AnimationWarping/Source/Runtime/Private/BoneControllers/AnimNode_FootPlacement.cpp b/Engine/Plugins/Animation/AnimationWarping/Source/Runtime/Private/BoneControllers/AnimNode_FootPlacement.cpp
--- a/Engine/Plugins/Animation/AnimationWarping/Source/Runtime/Private/BoneControllers/AnimNode_FootPlacement.cpp
+++ b/Engine/Plugins/Animation/AnimationWarping/Source/Runtime/Private/BoneControllers/AnimNode_FootPlacement.cpp
@@ -36,9 +36,9 @@ namespace UE::Anim::FootPlacement
 
 		FComponentSpacePoseContext& CSPContext;
 
-		//TODO: maybe store character?
 		class AActor* OwningActor = nullptr;
 		class UWorld* World = nullptr;
+		//@TODO: Make everything character related be input pins, to make this character mover agnostic. Not of the accesses are thread-safe technically.
 		class UCharacterMovementComponent* MovementComponent = nullptr;
 		FTransform OwningComponentToWorld = FTransform::Identity;
 		FTransform RootMotionTransformDelta = FTransform::Identity;
@@ -77,6 +77,11 @@ namespace UE::Anim::FootPlacement
 
 			return MovementComponent->IsWalkable(InHit);
 		}
+
+		FVector GetCharacterVelocity() const
+		{
+			return OwningActor->GetVelocity();
+		}
 	};
 
 	FEvaluationContext::FEvaluationContext(
@@ -434,7 +439,6 @@ void FAnimNode_FootPlacement::AlignPlantToGround(
 	const FVector IKReferenceNormalFootSpace = InputPoseAlignedRotationWS.UnrotateVector(PlantPlaneWS.GetNormal());
 
 	// Calculate and apply the amount of twist around the IK Root plane. 
-	// This is also used to calculate lock rotation limits
 	FQuat OutSwing;
 	UnalignedIKFootToUnalignedInputPoseRotationDelta.ToSwingTwist(IKReferenceNormalFootSpace,
 		OutSwing,
@@ -460,37 +464,23 @@ void FAnimNode_FootPlacement::AlignPlantToGround(
 
 FTransform FAnimNode_FootPlacement::UpdatePlantOffsetInterpolation(
 	const UE::Anim::FootPlacement::FEvaluationContext& Context,
-	UE::Anim::FootPlacement::FLegRuntimeData::FInterpolationData& InOutInterpData,
-	const FTransform& DesiredTransformCS) const
+	UE::Anim::FootPlacement::FLegRuntimeData::FInterpolationData& InOutInterpData) const
 {
-	const FTransform IkBase = DesiredTransformCS;
-	const FTransform IkBaseWithCurrOffset = IkBase * InOutInterpData.UnalignedFootOffsetCS;
-
-	// TODO: another indication everything should be done from the base
-	const FVector IkBaseTranslation = IkBaseWithCurrOffset.GetTranslation() - IkBase.GetTranslation();
-	const FVector IkBaseLastSpringTranslation = FVector(FVector2D(IkBaseTranslation), 0.0f);
-
 	const FVector IkBaseSpringTranslation = UKismetMathLibrary::VectorSpringInterp(
-		IkBaseLastSpringTranslation, FVector::ZeroVector, InOutInterpData.PlantOffsetTranslationSpringState,
+		InOutInterpData.UnalignedFootOffset.GetTranslation(), FVector::ZeroVector, InOutInterpData.PlantOffsetTranslationSpringState,
 		InterpolationSettings.UnplantLinearStiffness,
 		InterpolationSettings.UnplantLinearDamping,
 		Context.UpdateDeltaTime, 1.0f, 0.0f);
 
 	// Since the alignment is just a translation offset, there's no need to calculate a different offset.
 	const FQuat IkBaseSpringRotation = UKismetMathLibrary::QuaternionSpringInterp(
-		InOutInterpData.UnalignedFootOffsetCS.GetRotation(), FQuat::Identity, InOutInterpData.PlantOffsetRotationSpringState,
+		InOutInterpData.UnalignedFootOffset.GetRotation(), FQuat::Identity, InOutInterpData.PlantOffsetRotationSpringState,
 		InterpolationSettings.UnplantAngularStiffness,
 		InterpolationSettings.UnplantAngularDamping,
 		Context.UpdateDeltaTime, 1.0f, 0.0f);
 
-	//const FTransform IkBaseSpringOffset = FTransform(IkBaseSpringRotation, IkBaseSpringTranslation);
-	const FTransform IkBaseNewTransform = FTransform(IkBaseSpringRotation * IkBase.GetRotation(),
-		IkBase.GetLocation() + IkBaseSpringTranslation);
-	const FTransform IkBoneNewTransform = IkBaseNewTransform;
-
-	const FTransform BoneTransformOffset = IkBase.GetRelativeTransformReverse(IkBoneNewTransform);
-
-	return BoneTransformOffset;
+	const FTransform IkBaseSpringOffset = FTransform(IkBaseSpringRotation, IkBaseSpringTranslation);
+	return IkBaseSpringOffset;
 }
 
 void FAnimNode_FootPlacement::UpdatePlantingPlaneInterpolation(
@@ -499,39 +489,36 @@ void FAnimNode_FootPlacement::UpdatePlantingPlaneInterpolation(
 	const FTransform& LastAlignedFootTransform,
 	const float AlignmentAlpha,
 	FPlane& InOutPlantPlane,
+	const UE::Anim::FootPlacement::FLegRuntimeData::FInputPoseData& LegInputPose,
 	UE::Anim::FootPlacement::FLegRuntimeData::FInterpolationData& InOutInterpData) const
 {
 	const FVector TraceDirection = Context.ApproachDirWS;
 	const FPlane LastPlantPlane = InOutPlantPlane;
 
+	bool bFoundGround = false;
+	FVector ImpactLocationWS = FVector::ZeroVector;
 	if (TraceSettings.bDisableComplexTrace)
 	{
-		FVector ImpactLocationWS;
 		// Trace against simple geometry only
-		UE::Anim::FootPlacement::FindPlantPlane(Context, TraceSettings, FootTransformWS.GetLocation(), false, InOutPlantPlane, ImpactLocationWS);
-	}
-	else if (TraceSettings.SimpleCollisionInfluence > 0.0f && !FMath::IsNearlyEqual(AlignmentAlpha, 1.0f))
-	{
-		FVector ImpactLocationSimpleWS; 
-		FVector ImpactLocationComplexWS;
-
-		// Trace against both complex and simple geometry when the foot is in flight. Scale by alignment alpha.
-		// If the collision geometry we're testing against has simple and complex versions, simple collision may provide a smoother path with less clipping.
-		UE::Anim::FootPlacement::FindPlantPlane(Context, TraceSettings, FootTransformWS.GetLocation(), false, InOutPlantPlane, ImpactLocationSimpleWS);
-		UE::Anim::FootPlacement::FindPlantPlane(Context, TraceSettings, FootTransformWS.GetLocation(), true, InOutPlantPlane, ImpactLocationComplexWS);
-
-		// TODO: Alignment alpha is not really what we want. Once we have prediction, and know when the foot will be planted, we can build a better curve.
-		// Simple collision might be better for obstacle avoidance too, since it's presumably already a hull around complex collision.
-		FVector ImpactLocationBlendedWS = FMath::Lerp(
-			ImpactLocationSimpleWS, ImpactLocationComplexWS, 
-			AlignmentAlpha * TraceSettings.SimpleCollisionInfluence + (1.0f - TraceSettings.SimpleCollisionInfluence));
-		InOutPlantPlane = FPlane(ImpactLocationBlendedWS, InOutPlantPlane.GetNormal());
+		bFoundGround = UE::Anim::FootPlacement::FindPlantPlane(Context, TraceSettings, FootTransformWS.GetLocation(), false, InOutPlantPlane, ImpactLocationWS);
 	}
 	else
 	{
-		FVector ImpactLocationWS;
 		// Trace against complex geometry only to plant accurately
-		UE::Anim::FootPlacement::FindPlantPlane(Context, TraceSettings, FootTransformWS.GetLocation(), true, InOutPlantPlane, ImpactLocationWS);
+		bFoundGround = UE::Anim::FootPlacement::FindPlantPlane(Context, TraceSettings, FootTransformWS.GetLocation(), true, InOutPlantPlane, ImpactLocationWS);
+	}
+
+	if (CharacterData.bIsOnGround == false)
+	{
+		// if we're in the air, try to bring foot to source pose and collide with the ground.
+		const FPlane SourceFootPlane = FPlane(FootTransformWS.GetLocation(), Context.ApproachDirCS);
+		const float DistanceToGroundPlane =
+			UE::Anim::FootPlacement::GetDistanceToPlaneAlongDirection(ImpactLocationWS, SourceFootPlane, Context.ApproachDirCS);
+		const float PenetrationAmount = -DistanceToGroundPlane - LegInputPose.DistanceToPlant;
+		if (PenetrationAmount > 0.0f)
+		{
+			InOutPlantPlane = InOutPlantPlane.TranslateBy(-Context.ApproachDirCS * PenetrationAmount);
+		}
 	}
 
 	if (InterpolationSettings.bEnableFloorInterpolation && !bIsFirstUpdate)
@@ -553,9 +540,22 @@ void FAnimNode_FootPlacement::UpdatePlantingPlaneInterpolation(
 
 		const float LastPlaneDeltaZ = LastPlaneIntersection.Z - CurrPlaneIntersection.Z;
 		const float PrevPlaneDeltaZ = PrevPlaneIntersection.Z - CurrPlaneIntersection.Z;
-		const float AdjustedPrevZ = FMath::Abs(LastPlaneDeltaZ) < FMath::Abs(PrevPlaneDeltaZ) ?
+		float AdjustedPrevZ = FMath::Abs(LastPlaneDeltaZ) < FMath::Abs(PrevPlaneDeltaZ) ?
 			LastPlaneIntersection.Z : PrevPlaneIntersection.Z;
 
+		// Since our ground plane is in root space, apply the component delta if it gets us closer to the desired ground plane.
+		// This means if your ground geometry for capsule is smooth, we want to leverage that for foot alignment. 
+		// When ground geometry is not smooth, we follow the same logic as pelvis interpolation, and attempt to smooth out the jump in world space.
+		const float GroundPlaneDelta = CurrPlaneIntersection.Z  - AdjustedPrevZ;
+		if ((GroundPlaneDelta) > 0.0f)
+		{
+			AdjustedPrevZ += FMath::Clamp(-CharacterData.ComponentMoveDeltaWS.Z, 0.0f, GroundPlaneDelta);
+		}
+		else
+		{
+			AdjustedPrevZ += FMath::Clamp(-CharacterData.ComponentMoveDeltaWS.Z, GroundPlaneDelta, 0.0f);
+		}
+
 		float PlantPlaneSpringHeight = UKismetMathLibrary::FloatSpringInterp(
 			AdjustedPrevZ, CurrPlaneIntersection.Z, InOutInterpData.GroundHeightSpringState,
 			InterpolationSettings.FloorLinearStiffness,
@@ -563,11 +563,13 @@ void FAnimNode_FootPlacement::UpdatePlantingPlaneInterpolation(
 			Context.UpdateDeltaTime, 1.0f, 0.0f);
 
 		CurrPlaneIntersection.Z = PlantPlaneSpringHeight;
-		if (TraceSettings.MaxGroundPenetration >= 0.0f)
+
+		if ((TraceSettings.MaxGroundPenetration >= 0.0f) && bFoundGround)
 		{
 			// Prevent the foot from clipping too much into geometry due to interpolation
+			const FPlane GroundPlane = FPlane(ImpactLocationWS, InOutPlantPlane.GetNormal());
 			const float DistanceToGroundPlane =
-				UE::Anim::FootPlacement::GetDistanceToPlaneAlongDirection(CurrPlaneIntersection, InOutPlantPlane, Context.ApproachDirCS);
+				UE::Anim::FootPlacement::GetDistanceToPlaneAlongDirection(CurrPlaneIntersection, GroundPlane, Context.ApproachDirCS);
 			const float PenetrationAmount = -DistanceToGroundPlane - TraceSettings.MaxGroundPenetration;
 			if (PenetrationAmount > 0.0f)
 			{
@@ -753,7 +755,7 @@ UE::Anim::FootPlacement::FPlantResult FAnimNode_FootPlacement::FinalizeFootAlign
 	const FTransform& FKPelvisToHipCS =
 		LegData.InputPose.HipTransformCS.GetRelativeTransform(PelvisData.InputPose.FKTransformCS);
 	const FTransform FinalHipTransformCS = FKPelvisToHipCS * PelvisTransformCS;
-	FTransform CorrectedFootTransformCS = LegData.AlignedFootTransformCS;
+	FTransform CorrectedFootTransformCS = LegData.AlignedFootTransformRS * GetRootToComponent();
 	const FVector CorrectedBallLocationCS = (LegData.InputPose.FootToBall * CorrectedFootTransformCS).GetLocation();
 
 	// avoid hyper extension - start
@@ -846,13 +848,14 @@ UE::Anim::FootPlacement::FPlantResult FAnimNode_FootPlacement::FinalizeFootAlign
 	// and the tip must be calculated (note that because the ground plane interpolates, this does not prevent physical penetration 
 	// with the geometry).
 	// TODO: Consolidate with CalcTargetPlantPlaneDistance
+	const FPlane PlantPlaneCS = LegData.Plant.GetPlantPlaneCS(GetRootToComponent());
 	const float FootDistance = UE::Anim::FootPlacement::GetDistanceToPlaneAlongDirection(
 		CorrectedBallLocationCS,
-		LegData.Plant.PlantPlaneCS,
+		PlantPlaneCS,
 		Context.ApproachDirCS);
 	const float BallDistance = UE::Anim::FootPlacement::GetDistanceToPlaneAlongDirection(
 		CorrectedFootTransformCS.GetLocation(),
-		LegData.Plant.PlantPlaneCS,
+		PlantPlaneCS,
 		Context.ApproachDirCS);
 	const float MinDistance = FMath::Min(FootDistance, BallDistance);
 
@@ -880,7 +883,7 @@ UE::Anim::FootPlacement::FPlantResult FAnimNode_FootPlacement::FinalizeFootAlign
 		}
 	}
 
-	if (LegData.InputPose.DisableLeg > 0)
+	if (LegData.InputPose.DisableLeg > 0.0f)
 	{
 		FTransform DisabledLegTransform = LegData.InputPose.FootFKTransformCS;
 		DisabledLegTransform.SetTranslation(DisabledLegTransform.GetTranslation() + PelvisData.Interpolation.PelvisTranslationOffset * (1.0-PelvisData.DisablePelvis));
@@ -931,11 +934,12 @@ void FAnimNode_FootPlacement::DrawVLog(
 		LegData.InputPose.FootToGround *
 		LegData.AlignedFootTransformWS;
 
+	const FPlane PlantPlaneWS = LegData.Plant.GetPlantPlaneWS(GetRootToComponent(), Context.OwningComponentToWorld);
 
 	const FVector FKBoneLocationProjectedWS = UE::Anim::FootPlacement::PointDirectionPlaneIntersection(
 		FKBoneTransformWS.GetLocation(),
 		Context.ApproachDirWS,
-		LegData.Plant.PlantPlaneWS);
+		PlantPlaneWS);
 	
 	UObject* LogOwner = Context.CSPContext.AnimInstanceProxy->GetAnimInstanceObject();
 	FName LogCategory = "FootPlacement";
@@ -949,7 +953,7 @@ void FAnimNode_FootPlacement::DrawVLog(
 	const FVector IKBoneLocationProjectedWS = UE::Anim::FootPlacement::PointDirectionPlaneIntersection(
 		IKBoneTransformWS.GetLocation(),
 		Context.ApproachDirWS,
-		LegData.Plant.PlantPlaneWS);
+		PlantPlaneWS);
 
 	UE_VLOG_SPHERE(LogOwner, LogCategory, Display, IKBoneLocationProjectedWS, 0, CurrentPlantColor, TEXT(""));
 	UE_VLOG_SEGMENT_THICK(LogOwner, LogCategory, Display, IKBoneTransformWS.GetLocation(), IKBoneLocationProjectedWS, CurrentPlantColor, 2, TEXT(""));
@@ -958,9 +962,9 @@ void FAnimNode_FootPlacement::DrawVLog(
 	const FVector PlantCenter = UE::Anim::FootPlacement::PointDirectionPlaneIntersection(
 		IKBoneTransformWS.GetLocation(),
 		Context.ApproachDirWS,
-		LegData.Plant.PlantPlaneWS);
+		PlantPlaneWS);
 
-	UE_VLOG_CIRCLE(LogOwner, LogCategory, Display, PlantCenter, LegData.Plant.PlantPlaneWS.GetNormal(), UnplantRadius, PlantedColor, TEXT(""));
+	UE_VLOG_CIRCLE(LogOwner, LogCategory, Display, PlantCenter, PlantPlaneWS.GetNormal(), UnplantRadius, PlantedColor, TEXT(""));
 	
 	if (PlantSettings.ReplantRadiusRatio < 1.0f)
 	{
@@ -968,7 +972,7 @@ void FAnimNode_FootPlacement::DrawVLog(
 			PlantSettings.UnplantRadius *
 			PlantSettings.ReplantRadiusRatio;
 		
-		UE_VLOG_CIRCLE(LogOwner, LogCategory, Display, PlantCenter, LegData.Plant.PlantPlaneWS.GetNormal(), ReplantRadius, ReplantedColor, TEXT(""));
+		UE_VLOG_CIRCLE(LogOwner, LogCategory, Display, PlantCenter, PlantPlaneWS.GetNormal(), ReplantRadius, ReplantedColor, TEXT(""));
 	}
 
 	// FString InputPoseMessage = FString::Printf(
@@ -1042,11 +1046,12 @@ void FAnimNode_FootPlacement::DrawDebug(
 		LegData.InputPose.FootToGround *
 		LegData.AlignedFootTransformWS;
 
+	const FPlane PlantPlaneWS = LegData.Plant.GetPlantPlaneWS(GetRootToComponent(), Context.OwningComponentToWorld);
 
 	const FVector FKBoneLocationProjectedWS = UE::Anim::FootPlacement::PointDirectionPlaneIntersection(
 		FKBoneTransformWS.GetLocation(),
 		Context.ApproachDirWS,
-		LegData.Plant.PlantPlaneWS);
+		PlantPlaneWS);
 
 	Context.CSPContext.AnimInstanceProxy->AnimDrawDebugPoint(
 		FKBoneTransformWS.GetLocation(), 10.0f, FKColor, false, -1.0f, SDPG_Foreground);
@@ -1062,7 +1067,7 @@ void FAnimNode_FootPlacement::DrawDebug(
 	const FVector IKBoneLocationProjectedWS = UE::Anim::FootPlacement::PointDirectionPlaneIntersection(
 		IKBoneTransformWS.GetLocation(),
 		Context.ApproachDirWS,
-		LegData.Plant.PlantPlaneWS);
+		PlantPlaneWS);
 
 	Context.CSPContext.AnimInstanceProxy->AnimDrawDebugPoint(
 		IKBoneLocationProjectedWS, 15.0f, CurrentPlantColor, false, -1.0f, SDPG_Foreground);
@@ -1075,10 +1080,10 @@ void FAnimNode_FootPlacement::DrawDebug(
 	const FVector PlantCenter = UE::Anim::FootPlacement::PointDirectionPlaneIntersection(
 		IKBoneTransformWS.GetLocation(),
 		Context.ApproachDirWS,
-		LegData.Plant.PlantPlaneWS);
+		PlantPlaneWS);
 	Context.CSPContext.AnimInstanceProxy->AnimDrawDebugCircle(
 		PlantCenter, UnplantRadius, 24, PlantedColor,
-		LegData.Plant.PlantPlaneWS.GetNormal(), false, -1.0f, SDPG_Foreground, 0.5f);
+		PlantPlaneWS.GetNormal(), false, -1.0f, SDPG_Foreground, 0.5f);
 
 	if (PlantSettings.ReplantRadiusRatio < 1.0f)
 	{
@@ -1087,7 +1092,7 @@ void FAnimNode_FootPlacement::DrawDebug(
 			PlantSettings.ReplantRadiusRatio;
 		Context.CSPContext.AnimInstanceProxy->AnimDrawDebugCircle(
 			PlantCenter, ReplantRadius, 24, ReplantedColor,
-			LegData.Plant.PlantPlaneWS.GetNormal(), false, -1.0f, SDPG_Foreground, 0.5f);
+			PlantPlaneWS.GetNormal(), false, -1.0f, SDPG_Foreground, 0.5f);
 	}
 
 	FString InputPoseMessage = FString::Printf(
@@ -1164,7 +1169,7 @@ void FAnimNode_FootPlacement::EvaluateSkeletalControl_AnyThread(FComponentSpaceP
 		const FTransform ComponentTransform = Output.AnimInstanceProxy->GetComponentTransform();
 		const FVector RootWorldSpaceLocation = ComponentTransform.TransformPosition(
 			Output.Pose.GetComponentSpaceTransform(FCompactPoseBoneIndex(0)).GetTranslation());
-		const FVector PrevRootWorldSpaceLocation = CharacterData.ComponentTransformWS.TransformPosition(PelvisData.InputPose.RootTransformCS.GetTranslation());
+		const FVector PrevRootWorldSpaceLocation = CharacterData.ComponentTransformWS.TransformPosition(GetRootToComponent().GetTranslation());
 		if (FVector::DistSquared(RootWorldSpaceLocation, PrevRootWorldSpaceLocation) > FMath::Square(TeleportDistanceThreshold))
 		{
 			ResetRuntimeData();
@@ -1221,7 +1226,10 @@ void FAnimNode_FootPlacement::EvaluateSkeletalControl_AnyThread(FComponentSpaceP
 #endif
 	if (PelvisSettings.bEnableInterpolation)
 	{
-		PelvisTransformCS = UpdatePelvisInterpolation(FootPlacementContext, PelvisTransformCS);
+		const FTransform& RootToComponent = GetRootToComponent();
+		const FTransform TargetPelvisTransformRS = PelvisTransformCS.GetRelativeTransform(RootToComponent);
+		const FTransform PelvisTransformRS = UpdatePelvisInterpolationRootSpace(FootPlacementContext, TargetPelvisTransformRS);
+		PelvisTransformCS = PelvisTransformRS * RootToComponent;
 	}
 	
 	PelvisData.DisablePelvis = Output.Curve.Get(PelvisSettings.DisablePelvisCurveName);
@@ -1440,6 +1448,7 @@ void FAnimNode_FootPlacement::GatherPelvisDataFromInputs(const UE::Anim::FootPla
 		Context.CSPContext.Pose.GetComponentSpaceTransform(PelvisData.Bones.FkBoneIndex);
 	PelvisData.InputPose.IKRootTransformCS =
 		Context.CSPContext.Pose.GetComponentSpaceTransform(PelvisData.Bones.IkBoneIndex);
+
 	PelvisData.InputPose.RootTransformCS =
 		Context.CSPContext.Pose.GetComponentSpaceTransform(FCompactPoseBoneIndex(0));
 
@@ -1477,15 +1486,19 @@ void FAnimNode_FootPlacement::GatherLegDataFromInputs(
 
 	if (bIsFirstUpdate)
 	{
+		const FTransform& RootToComponent = GetRootToComponent();
+
+		LegData.AlignedFootTransformRS = LegData.InputPose.FootTransformCS.GetRelativeTransform(RootToComponent);
 		LegData.AlignedFootTransformWS =
+		LegData.UnalignedFootTransformRS = LegData.AlignedFootTransformRS;
 			LegData.InputPose.FootTransformCS * Context.OwningComponentToWorld;
 		LegData.UnalignedFootTransformWS = LegData.AlignedFootTransformWS;
 
 		const FVector IKFootRootLocationWS =
 			Context.OwningComponentToWorld.TransformPosition(PelvisData.InputPose.IKRootTransformCS.GetLocation());
 
-		LegData.Plant.PlantPlaneWS = FPlane(IKFootRootLocationWS, -Context.ApproachDirWS);
-		LegData.Plant.PlantPlaneCS = FPlane(PelvisData.InputPose.IKRootTransformCS.GetLocation(), -Context.ApproachDirCS);
+		const FVector IKFootRootLocationRS = RootToComponent.InverseTransformPosition(PelvisData.InputPose.IKRootTransformCS.GetLocation());
+		LegData.Plant.PlantPlaneRS = FPlane(IKFootRootLocationRS, -Context.ApproachDirCS);
 
 		LegData.Plant.PlantType = UE::Anim::FootPlacement::EPlantType::Unplanted;
 		LegData.Plant.LastPlantType = UE::Anim::FootPlacement::EPlantType::Unplanted;
@@ -1539,6 +1552,12 @@ void FAnimNode_FootPlacement::ProcessCharacterState(const UE::Anim::FootPlacemen
 		? Context.OwningComponentToWorld.GetLocation()
 		: CharacterData.ComponentTransformWS.GetLocation();
 
+	if (bIsFirstUpdate)
+	{
+		CharacterData.SmoothCapsuleGroundNormalWS = -Context.ApproachDirWS;
+		CharacterData.SmoothCapsuleGroundNormalSpringState.Reset();
+	}
+
 	CharacterData.ComponentTransformWS = Context.OwningComponentToWorld;
 	const FVector ComponentLocationWS = CharacterData.ComponentTransformWS.GetLocation();
 
@@ -1549,14 +1568,27 @@ void FAnimNode_FootPlacement::ProcessCharacterState(const UE::Anim::FootPlacemen
 			(Context.MovementComponent->MovementMode == MOVE_NavWalking)) &&
 		Context.MovementComponent->CurrentFloor.bBlockingHit;
 
+	CharacterData.ComponentMoveDeltaWS = FVector::ZeroVector;
 	bool bOnGround =  !PelvisSettings.bDisablePelvisOffsetInAir || (CharacterData.bIsOnGround && bWasOnGround);
 	if (bOnGround && PelvisSettings.ActorMovementCompensationMode != EActorMovementCompensationMode::ComponentSpace)
 	{
 		FVector OwningComponentAdjustedLastLocationWS;
 		if (PelvisSettings.ActorMovementCompensationMode == EActorMovementCompensationMode::SuddenMotionOnly)
 		{
+			FQuat SlopeDelta = FQuat::FindBetweenNormals(CharacterData.SmoothCapsuleGroundNormalWS, Context.GetMovementComponentFloorNormal());
+			SlopeDelta = UKismetMathLibrary::QuaternionSpringInterp(
+																FQuat::Identity, 
+																SlopeDelta, 
+																CharacterData.SmoothCapsuleGroundNormalSpringState, 
+																InterpolationSettings.FloorAngularStiffness, 
+																1.0f, 
+																Context.UpdateDeltaTime, 
+																1.0f, 
+																0.0f);
+			CharacterData.SmoothCapsuleGroundNormalWS = SlopeDelta.RotateVector(CharacterData.SmoothCapsuleGroundNormalWS);
+
 			// Compensate for sudden capsule moves
-			const FVector CapsuleFloorNormalWS = Context.GetMovementComponentFloorNormal();
+			const FVector CapsuleFloorNormalWS = CharacterData.SmoothCapsuleGroundNormalWS;
 			OwningComponentAdjustedLastLocationWS =
 				(FMath::Abs(Context.ApproachDirWS | CapsuleFloorNormalWS) > DELTA) ?
 				UE::Anim::FootPlacement::PointDirectionPlaneIntersection(
@@ -1573,7 +1605,9 @@ void FAnimNode_FootPlacement::ProcessCharacterState(const UE::Anim::FootPlacemen
 
 		// Only compensate vertical motion
 		const FVector CapsuleMoveOffsetWS =
-			(ComponentLocationWS - OwningComponentAdjustedLastLocationWS) * -Context.ApproachDirWS;
+			(ComponentLocationWS - OwningComponentAdjustedLastLocationWS - BaseTranslationDelta) * -Context.ApproachDirWS;
+
+		CharacterData.ComponentMoveDeltaWS +=CapsuleMoveOffsetWS;
 		if (!CapsuleMoveOffsetWS.IsNearlyZero(KINDA_SMALL_NUMBER))
 		{
 			const FVector CapsuleMoveOffsetCS =
@@ -1584,20 +1618,14 @@ void FAnimNode_FootPlacement::ProcessCharacterState(const UE::Anim::FootPlacemen
 	}
 
 	{
-		// Use the component's move delta, instead of the movement component's velocity, since it doesn't account for uphill/downhill velocity
+		CharacterData.CharacterVelocityWS = Context.GetCharacterVelocity();
+
+		// Also grab the component's move delta, instead of the movement component's velocity, since it doesn't account for uphill/downhill velocity
 		const FVector CapsuleMoveOffsetWS =
 			(ComponentLocationWS - LastComponentLocationWS);
 		const FVector CapsuleMoveOffsetCS =
 			Context.OwningComponentToWorld.InverseTransformVectorNoScale(CapsuleMoveOffsetWS);
-
-		if (CachedDeltaTime > 0.0f)
-		{
-			CharacterData.ComponentVelocityCS = CapsuleMoveOffsetCS / CachedDeltaTime;
-		}
-		else
-		{
-			CharacterData.ComponentVelocityCS = FVector::ZeroVector;
-		}
+		CharacterData.ComponentMoveDeltaWS += CapsuleMoveOffsetWS;
 	}
 }
 
@@ -1609,10 +1637,27 @@ void FAnimNode_FootPlacement::ProcessFootAlignment(
 	UE::Anim::FootPlacement::FLegRuntimeData::FInterpolationData& Interpolation = LegData.Interpolation;
 	UE::Anim::FootPlacement::FLegRuntimeData::FBoneData& Bones = LegData.Bones;
 	UE::Anim::FootPlacement::FLegRuntimeData::FPlantData& Plant = LegData.Plant;
+	
+	const FTransform& RootToComponent = GetRootToComponent();
+	if (PlantSettings.bReconstructWorldPlantFromVelocity)
+	{
+		// Last frame's plant in root space minus our move delta from character velocity.
+		LegData.AlignedFootTransformWS = LegData.AlignedFootTransformRS * RootToComponent * Context.OwningComponentToWorld;
+		LegData.AlignedFootTransformWS.AddToTranslation(-CharacterData.CharacterVelocityWS * Context.UpdateDeltaTime);
+
+		LegData.UnalignedFootTransformWS = LegData.UnalignedFootTransformRS * RootToComponent * Context.OwningComponentToWorld;
+		LegData.UnalignedFootTransformWS.AddToTranslation(-CharacterData.CharacterVelocityWS * Context.UpdateDeltaTime);
+	}
+	else
+	{
+		LegData.AlignedFootTransformWS.AddToTranslation(-BaseTranslationDelta);
+		LegData.UnalignedFootTransformWS.AddToTranslation(-BaseTranslationDelta);
+	}
 
 	const FTransform InputPoseFootTransformWS = InputPose.FootTransformCS * Context.OwningComponentToWorld;
 	const FTransform LastAlignedFootTransformWS = LegData.AlignedFootTransformWS;
 	const FTransform LastUnalignedFootTransformWS = LegData.UnalignedFootTransformWS;
+	const FTransform InputPoseFootTransformRS = InputPose.FootTransformCS.GetRelativeTransform(RootToComponent);
 
 	Plant.LastPlantType = Plant.PlantType;
 	DeterminePlantType(
@@ -1656,20 +1701,21 @@ void FAnimNode_FootPlacement::ProcessFootAlignment(
 		default: check(false); break; //not implemented
 		}
 
-		FTransform PlantedFootTransformCS =
+		const FTransform PlantedFootTransformCS =
 			CurrentPlantedTransformWS * Context.OwningComponentToWorld.Inverse();
+		FTransform PlantedFootTransformRS = PlantedFootTransformCS.GetRelativeTransform(RootToComponent);
 
 		// The locked transform is aligned to the ground. Conserve the input pose's ground alignment
-		const FVector AlignedBoneLocationCS = PlantedFootTransformCS.GetLocation();
-		const FPlane InputPosePlantPlane = FPlane(InputPose.FootTransformCS.GetLocation(), Context.ApproachDirWS);
-		const FVector UnalignedBoneLocationCS = FVector::PointPlaneProject(AlignedBoneLocationCS, InputPosePlantPlane);
+		const FVector AlignedBoneLocationRS = PlantedFootTransformRS.GetLocation();
+		const FPlane InputPosePlantPlane = FPlane(InputPoseFootTransformRS.GetLocation(), Context.ApproachDirWS);
+		const FVector UnalignedBoneLocationRS = FVector::PointPlaneProject(AlignedBoneLocationRS, InputPosePlantPlane);
 
-		PlantedFootTransformCS.SetLocation(UnalignedBoneLocationCS);
+		PlantedFootTransformRS.SetLocation(UnalignedBoneLocationRS);
 
 		// Get the offset relative to the initial foot transform
 		// Reset interpolation 
-		Interpolation.UnalignedFootOffsetCS =
-			InputPose.FootTransformCS.GetRelativeTransformReverse(PlantedFootTransformCS);
+		Interpolation.UnalignedFootOffset =
+			InputPoseFootTransformRS.GetRelativeTransformReverse(PlantedFootTransformRS);
 		Interpolation.PlantOffsetTranslationSpringState.Reset();
 		Interpolation.PlantOffsetRotationSpringState.Reset();
 
@@ -1679,8 +1725,8 @@ void FAnimNode_FootPlacement::ProcessFootAlignment(
 	else
 	{
 		// No plant, so we interpolate the offset out
-		Interpolation.UnalignedFootOffsetCS =
-			UpdatePlantOffsetInterpolation(Context, Interpolation, InputPose.FootTransformCS);
+		Interpolation.UnalignedFootOffset =
+			UpdatePlantOffsetInterpolation(Context, Interpolation);
 
 		// If we're unplanted, we know we're fully unaligned the first time we hit zero alignment alpha.
 		if (Plant.TimeSinceFullyUnaligned > 0.0f || FMath::IsNearlyZero(InputPose.AlignmentAlpha))
@@ -1692,14 +1738,14 @@ void FAnimNode_FootPlacement::ProcessFootAlignment(
 	// If replant radius is the same as unplant radius, clamp the location and slide
 	if (PlantSettings.ReplantRadiusRatio >= 1.0f)
 	{
-		const FVector ClampedTransltionOffset = Interpolation.UnalignedFootOffsetCS.GetLocation().GetClampedToMaxSize(PlantSettings.UnplantRadius);
-		Interpolation.UnalignedFootOffsetCS.SetLocation(ClampedTransltionOffset);
+		const FVector ClampedTransltionOffset = Interpolation.UnalignedFootOffset.GetLocation().GetClampedToMaxSize(PlantSettings.UnplantRadius);
+		Interpolation.UnalignedFootOffset.SetLocation(ClampedTransltionOffset);
 	}
 
 	// If replant angle is the same as unplant angle, clamp the angle and slide
 	if (PlantSettings.ReplantAngleRatio >= 1.0f)
 	{
-		FQuat ClampedRotationOffset = Interpolation.UnalignedFootOffsetCS.GetRotation();
+		FQuat ClampedRotationOffset = Interpolation.UnalignedFootOffset.GetRotation();
 		ClampedRotationOffset.Normalize();
 		ClampedRotationOffset = ClampedRotationOffset.W < 0.0 ? -ClampedRotationOffset : ClampedRotationOffset;
 
@@ -1712,21 +1758,22 @@ void FAnimNode_FootPlacement::ProcessFootAlignment(
 		{
 			ClampedRotationOffset = FQuat(OffsetAxis, MaxAngle);
 		}
-		Interpolation.UnalignedFootOffsetCS.SetRotation(ClampedRotationOffset);
+		Interpolation.UnalignedFootOffset.SetRotation(ClampedRotationOffset);
 	}
 
-	FTransform FootUnalignedTransformCS = InputPose.FootTransformCS * Interpolation.UnalignedFootOffsetCS;
+	FTransform FootUnalignedTransformRS = InputPose.FootTransformCS.GetRelativeTransform(RootToComponent) * Interpolation.UnalignedFootOffset;
+	//@TODO: Change this logic to be in root space.
 	if (PlantSettings.SeparatingDistance > 0.0f)
 	{
 		// Prevent the feet from crossing by enforcing a set distance from a plane at the midpoint between all feet
-		FVector FootUnalignedLocationCS = FootUnalignedTransformCS.GetLocation();
+		FVector FootUnalignedLocationRS = FootUnalignedTransformRS.GetLocation();
 		const FVector MidPointToFoot = (InputPose.FootTransformCS.GetLocation() - PelvisData.InputPose.FootMidpointCS);
 		const FVector PlaneNormal = MidPointToFoot.GetSafeNormal2D();
 		const FVector PlaneCenter = PelvisData.InputPose.FootMidpointCS + PlaneNormal * PlantSettings.SeparatingDistance;
 		const FPlane SeparatingPlane = FPlane(PlaneCenter, PlaneNormal);
 		
 		const float DistanceToSeparatingPlane =
-			UE::Anim::FootPlacement::GetDistanceToPlaneAlongDirection(FootUnalignedLocationCS, SeparatingPlane, -PlaneNormal);
+			UE::Anim::FootPlacement::GetDistanceToPlaneAlongDirection(FootUnalignedLocationRS, SeparatingPlane, -PlaneNormal);
 
 		if (Plant.PlantType == UE::Anim::FootPlacement::EPlantType::Unplanted)
 		{
@@ -1749,8 +1796,8 @@ void FAnimNode_FootPlacement::ProcessFootAlignment(
 				Interpolation.SeparatingPlaneOffset = SeparatingPlaneOffset;
 			}
 			
-			FootUnalignedLocationCS  += Interpolation.SeparatingPlaneOffset;
-			FootUnalignedTransformCS.SetLocation(FootUnalignedLocationCS);
+			FootUnalignedLocationRS  += Interpolation.SeparatingPlaneOffset;
+			FootUnalignedTransformRS.SetLocation(FootUnalignedLocationRS);
 		}
 		else
 		{
@@ -1782,45 +1829,48 @@ void FAnimNode_FootPlacement::ProcessFootAlignment(
 #endif
 	}
 
-	FTransform BlendedUnalignedTransformCS;
+	FTransform BlendedUnalignedTransformRS;
 	{
 		// Blend the component-space input pose, with the unaligned foot-locked transform.
 		// Allow ground alignment to continue with this blended result.
 		// When the lock alpha reaches 0, we will automatically unlock the foot.
-		BlendedUnalignedTransformCS.Blend(InputPose.FootTransformCS, FootUnalignedTransformCS, InputPose.LockAlpha);
-		LegData.UnalignedFootTransformWS = BlendedUnalignedTransformCS * Context.OwningComponentToWorld;
+		BlendedUnalignedTransformRS.Blend(InputPoseFootTransformRS, FootUnalignedTransformRS, InputPose.LockAlpha);
+		LegData.UnalignedFootTransformRS = BlendedUnalignedTransformRS;
+		LegData.UnalignedFootTransformWS = BlendedUnalignedTransformRS * RootToComponent * Context.OwningComponentToWorld;
 	}
 
 	const FTransform ComponentToWorldInv = Context.OwningComponentToWorld.Inverse();
 
 	// find the smooth plant plane
-	UpdatePlantingPlaneInterpolation(Context, LegData.UnalignedFootTransformWS, LastAlignedFootTransformWS, InputPose.AlignmentAlpha, Plant.PlantPlaneWS, Interpolation);
-	Plant.PlantPlaneCS = Plant.PlantPlaneWS.TransformBy(ComponentToWorldInv.ToMatrixWithScale());
-
-	Interpolation.UnalignedFootOffsetCS = InputPose.FootTransformCS.GetRelativeTransformReverse(FootUnalignedTransformCS);
+	FPlane PlantPlaneWS = Plant.GetPlantPlaneWS(RootToComponent, Context.OwningComponentToWorld);
+	UpdatePlantingPlaneInterpolation(Context, LegData.UnalignedFootTransformWS,
+									LastAlignedFootTransformWS, 
+									InputPose.AlignmentAlpha, 
+		PlantPlaneWS, 
+									InputPose, 
+									Interpolation);
+	const FPlane PlantPlaneCS = PlantPlaneWS.TransformBy(ComponentToWorldInv.ToMatrixWithScale());
+	Plant.PlantPlaneRS = PlantPlaneCS.TransformBy(RootToComponent.Inverse().ToMatrixWithScale());
 
 	// This will adjust UnalignedFootTransformWS to make it match the required distance to the plant plane along the 
 	// approach direction, not the plane normal
 	LegData.AlignedFootTransformWS = LegData.UnalignedFootTransformWS;
-	AlignPlantToGround(Context, Plant.PlantPlaneWS, InputPose, LegData.AlignedFootTransformWS, Plant.TwistCorrection);
-
-	LegData.AlignedFootTransformCS =
-		LegData.AlignedFootTransformWS * ComponentToWorldInv;
+	AlignPlantToGround(Context, PlantPlaneWS, InputPose, LegData.AlignedFootTransformWS, Plant.TwistCorrection);
 
+	const FTransform AlignedFootTransformCS =	LegData.AlignedFootTransformWS * ComponentToWorldInv;
 	// The target transform is a blend based on FkAlignmentAlpha.
 	// Until we have prediction, favor the ground aligned position, 
 	// since this will likely have a more accurate distance from plane
-	FTransform BlendedPlantTransformCS = LegData.AlignedFootTransformCS;
-
+	FTransform BlendedPlantTransformCS = AlignedFootTransformCS;
 	// When unplanted/unaligned, favor FK orientation and fix penetrations later. 
 	BlendedPlantTransformCS.SetRotation(
 		FQuat::Slerp(
 			InputPose.FootTransformCS.GetRotation(),
-			LegData.AlignedFootTransformCS.GetRotation(),
+			AlignedFootTransformCS.GetRotation(),
 			InputPose.AlignmentAlpha));
 
-	
-	LegData.AlignedFootTransformCS = BlendedPlantTransformCS;
+	LegData.AlignedFootTransformRS = BlendedPlantTransformCS.GetRelativeTransform(RootToComponent);
+	LegData.AlignedFootTransformWS = BlendedPlantTransformCS * Context.OwningComponentToWorld;
 }
 
 FVector FAnimNode_FootPlacement::GetApproachDirWS(const FAnimationBaseContext& Context) const
@@ -1829,13 +1879,15 @@ FVector FAnimNode_FootPlacement::GetApproachDirWS(const FAnimationBaseContext& C
 	return -(OwningComponent->GetComponentTransform().GetRotation().GetUpVector());
 }
 
+const FTransform& FAnimNode_FootPlacement::GetRootToComponent() const
+{
+	return PelvisData.InputPose.RootTransformCS;
+}
+
 FTransform FAnimNode_FootPlacement::SolvePelvis(const UE::Anim::FootPlacement::FEvaluationContext& Context)
 {
 	using namespace UE::Anim::FootPlacement;
 
-	TBitArray<> RelevantFeet = FindRelevantFeet(Context);
-	const int32 RelevantFeetNum = RelevantFeet.CountSetBits();
-
 	// Rebalance the pelvis before calculating its desired height
 	FTransform RebalancedPelvisTransform  = PelvisData.InputPose.FKTransformCS;
 	FVector PelvisOffsetDelta;
@@ -1845,7 +1897,7 @@ FTransform FAnimNode_FootPlacement::SolvePelvis(const UE::Anim::FootPlacement::F
 		FVector OffsetAverage = FVector::ZeroVector;
 		for (const FLegRuntimeData& LegData : LegsData)
 		{
-			const FVector LegTranslationOffset = LegData.AlignedFootTransformCS.GetLocation() - LegData.InputPose.FootTransformCS.GetLocation();
+			const FVector LegTranslationOffset = GetRootToComponent().TransformPosition(LegData.AlignedFootTransformRS.GetLocation()) - LegData.InputPose.FootTransformCS.GetLocation();
 			OffsetAverage += LegTranslationOffset / NumLegs;
 		}
 
@@ -1862,16 +1914,14 @@ FTransform FAnimNode_FootPlacement::SolvePelvis(const UE::Anim::FootPlacement::F
 	float DesiredOffsetAvg = 0.0f;
 	float MinOffsetMax = -BIG_NUMBER;
 
-	for (TConstSetBitIterator It(RelevantFeet); It; ++It)
+	const int32 FootNum = LegsData.Num();
+	for (const FLegRuntimeData& LegData : LegsData)
 	{
-		const int32 LegIndex = It.GetIndex();
-		const FLegRuntimeData& LegData = LegsData[LegIndex];
-
 		FPelvisOffsetRangeForLimb PelvisOffsetRangeCS;
 		FindPelvisOffsetRangeForLimb(
 			Context,
 			LegData,
-			LegData.AlignedFootTransformCS.GetLocation(),
+			GetRootToComponent().TransformPosition(LegData.AlignedFootTransformRS.GetLocation()),
 			RebalancedPelvisTransform,
 			PelvisOffsetRangeCS);
 
@@ -1879,7 +1929,7 @@ FTransform FAnimNode_FootPlacement::SolvePelvis(const UE::Anim::FootPlacement::F
 		const float MaxOffset = PelvisOffsetRangeCS.MaxExtension;
 		const float MinOffset = PelvisOffsetRangeCS.MinExtension;
 
-		DesiredOffsetAvg += DesiredOffset / RelevantFeetNum;
+		DesiredOffsetAvg += DesiredOffset / FootNum;
 		DesiredOffsetMin = FMath::Min(DesiredOffsetMin, DesiredOffset);
 		MaxOffsetMin = FMath::Min(MaxOffsetMin, MaxOffset);
 		MinOffsetMax = FMath::Max(MinOffsetMax, MinOffset);
@@ -1906,92 +1956,17 @@ FTransform FAnimNode_FootPlacement::SolvePelvis(const UE::Anim::FootPlacement::F
 	return PelvisTransformCS;
 }
 
-TBitArray<> FAnimNode_FootPlacement::FindRelevantFeet(const UE::Anim::FootPlacement::FEvaluationContext& Context)
-{
-	TBitArray<> RelevantFeet;
-
-	if (PelvisSettings.PelvisHeightMode == EPelvisHeightMode::AllLegs)
-	{
-		RelevantFeet.Init(true, LegsData.Num());
-		return RelevantFeet;
-	}
-
-
-	if (PelvisSettings.PelvisHeightMode == EPelvisHeightMode::FrontPlantedFeetUphill_FrontFeetDownhill)
-	{
-		if (!CharacterData.ComponentVelocityCS.IsNearlyZero())
-		{
-			TBitArray<> FrontFeet;
-			FrontFeet.SetNumUninitialized(LegsData.Num());
-			const FVector MovementDirection = CharacterData.ComponentVelocityCS.GetUnsafeNormal();
-			for (int32 LegIndex = 0; LegIndex < LegsData.Num(); ++LegIndex)
-			{
-				const FVector FootLocation = LegsData[LegIndex].InputPose.FootTransformCS.GetLocation();
-				const FVector MidpointToFoot = LegsData[LegIndex].InputPose.FootTransformCS.GetLocation() - PelvisData.InputPose.FootMidpointCS;
-				FrontFeet[LegIndex] = MidpointToFoot.Dot(MovementDirection) > 0.0f;
-			}
-
-			const int32 NumFrontFeet = FrontFeet.CountSetBits();
-			const bool bMovingDownHill = CharacterData.ComponentVelocityCS.Dot(Context.ApproachDirCS) > 0.0f;
-			if (bMovingDownHill)
-			{
-				// Consider all legs relevant when moving downhill. When using multiple feet, we tend to favour the lowest foot anyway.
-				RelevantFeet.Init(true, LegsData.Num());
-			}
-			else
-			{
-				TBitArray<> PlantedFeet = FindPlantedFeet(Context);
-				// Prefer the front-most planted foot, if any
-				RelevantFeet = TBitArray<>::BitwiseAND(PlantedFeet, FrontFeet , EBitwiseOperatorFlags::MinSize);
-
-				const int32 NumRelevantFeet = RelevantFeet.CountSetBits();
-				if (NumRelevantFeet == 0)
-				{
-					// No front-most planted feet. Fallback to any planted feet
-					RelevantFeet = PlantedFeet;
-				}
-			}
-		}
-	}
-	else //if (PelvisSettings.PelvisHeightMode == EPelvisHeightMode::AllPlantedFeet)
-	{
-		TBitArray<> PlantedFeet = FindPlantedFeet(Context);
-		RelevantFeet = PlantedFeet;
-	}
-	
-	if (RelevantFeet.CountSetBits() == 0)
-	{
-		// If we reach this point with no planted feet, fall-back to all feet
-		RelevantFeet.Init(true, LegsData.Num());
-	}
-
-	return RelevantFeet;
-}
-
-TBitArray<> FAnimNode_FootPlacement::FindPlantedFeet(const UE::Anim::FootPlacement::FEvaluationContext& Context)
-{
-	TBitArray<> PlantedFeet;
-	PlantedFeet.Init(true, LegsData.Num());
-	for (int32 LegIndex = 0; LegIndex < LegsData.Num(); ++LegIndex)
-	{
-		// Using alignment alpha, instead of bWantsToPlant allows us to start the weight-shift a few frames before planting
-		if (LegsData[LegIndex].InputPose.AlignmentAlpha == 0.0f)
-		{
-			PlantedFeet[LegIndex] = false;
-		}
-	}
-
-	return PlantedFeet;
-}
-
-FTransform FAnimNode_FootPlacement::UpdatePelvisInterpolation(
+FTransform FAnimNode_FootPlacement::UpdatePelvisInterpolationRootSpace(
 	const UE::Anim::FootPlacement::FEvaluationContext& Context,
-	const FTransform& TargetPelvisTransform)
+	const FTransform& TargetPelvisTransformRS)
 {
-	FTransform OutPelvisTransform = TargetPelvisTransform;
+	const FTransform& RootTransformCS = GetRootToComponent();
+	const FVector PelvisLocationRS = RootTransformCS.InverseTransformPosition(PelvisData.InputPose.FKTransformCS.GetLocation());
+
+	FTransform OutPelvisTransform = TargetPelvisTransformRS;
 	// Calculate the offset from input pose and interpolate
 	FVector DesiredPelvisOffset =
-		TargetPelvisTransform.GetLocation() - PelvisData.InputPose.FKTransformCS.GetLocation();
+		TargetPelvisTransformRS.GetLocation() - PelvisLocationRS;
 
 	// Clamp by MaxOffset
 	// Clamping the target before interpolation means we may exceed this purely do to interpolation.
@@ -2012,7 +1987,7 @@ FTransform FAnimNode_FootPlacement::UpdatePelvisInterpolation(
 	PelvisData.Interpolation.PelvisTranslationOffset = NewTranslationOffset;
 
 	OutPelvisTransform.SetLocation(
-		PelvisData.InputPose.FKTransformCS.GetLocation() + PelvisData.Interpolation.PelvisTranslationOffset);
+		PelvisLocationRS + PelvisData.Interpolation.PelvisTranslationOffset);
 
 	return OutPelvisTransform;
 }
diff --git a/Engine/Plugins/Animation/AnimationWarping/Source/Runtime/Public/BoneControllers/AnimNode_FootPlacement.h b/Engine/Plugins/Animation/AnimationWarping/Source/Runtime/Public/BoneControllers/AnimNode_FootPlacement.h
--- a/Engine/Plugins/Animation/AnimationWarping/Source/Runtime/Public/BoneControllers/AnimNode_FootPlacement.h
+++ b/Engine/Plugins/Animation/AnimationWarping/Source/Runtime/Public/BoneControllers/AnimNode_FootPlacement.h
@@ -44,12 +44,6 @@ namespace UE::Anim::FootPlacement
 			float FootLength = 0.0f;
 		} Bones;
 
-		// Curves
-		UE_DEPRECATED(5.3, "SpeedCurveUID is no longer used, use SpeedCurveName instead.")
-		SmartName::UID_Type SpeedCurveUID = SmartName::MaxUID;
-		UE_DEPRECATED(5.3, "DisableLockCurveUID is no longer used, use DisableLockCurveName instead.")
-		SmartName::UID_Type DisableLockCurveUID = SmartName::MaxUID;
-
 		FName SpeedCurveName = NAME_None;
 		FName DisableLockCurveName = NAME_None;
 		FName DisableLegCurveName = NAME_None;
@@ -84,8 +78,7 @@ namespace UE::Anim::FootPlacement
 		{
 			UE::Anim::FootPlacement::EPlantType PlantType = UE::Anim::FootPlacement::EPlantType::Unplanted;
 			UE::Anim::FootPlacement::EPlantType LastPlantType = UE::Anim::FootPlacement::EPlantType::Unplanted;
-			FPlane PlantPlaneWS = FPlane(FVector::UpVector, 0.0f);
-			FPlane PlantPlaneCS = FPlane(FVector::UpVector, 0.0f);
+			FPlane PlantPlaneRS = FPlane(FVector::UpVector, 0.0f);
 			FQuat TwistCorrection = FQuat::Identity;
 			// @TODO: When we have prediction/phase info, replace use-cases with post-plant roll-phase
 			float TimeSinceFullyUnaligned = 0.0f;
@@ -93,19 +86,27 @@ namespace UE::Anim::FootPlacement
 			bool bCanReachTarget = false;
 			// Whether we want to plant, independently from any dynamic pose adjustments we may do
 			bool bWantsToPlant = false;
+
+			FPlane GetPlantPlaneCS(const FTransform& RootToComponent) const { return PlantPlaneRS.TransformBy(RootToComponent.ToMatrixWithScale());}
+			FPlane GetPlantPlaneWS(const FTransform& RootToComponent, const FTransform& ComponentToWorld) const 
+			{ 
+				return GetPlantPlaneCS(RootToComponent).TransformBy(ComponentToWorld.ToMatrixWithScale());
+			}
+
 		} Plant;
 		
 		// Ground-aligned, locked/unlocked bone transform pre-extension adjustments
 		FTransform AlignedFootTransformWS = FTransform::Identity;
-		FTransform AlignedFootTransformCS = FTransform::Identity;
+		FTransform AlignedFootTransformRS = FTransform::Identity;
 		// Foot locked/unlocked bone transform, before ground alignment
 		FTransform UnalignedFootTransformWS = FTransform::Identity;
+		FTransform UnalignedFootTransformRS = FTransform::Identity;
 		
 		/* Interpolation */
 		struct FInterpolationData
 		{
 			// Interpolated foot lock offset
-			FTransform UnalignedFootOffsetCS = FTransform::Identity;
+			FTransform UnalignedFootOffset = FTransform::Identity;
 			// Separating plane spring states
 			FVectorSpringState SeparatingPlaneOffsetSpringState;
 			FVector SeparatingPlaneOffset = FVector::ZeroVector;
@@ -162,7 +163,10 @@ namespace UE::Anim::FootPlacement
 	struct FCharacterData
 	{
 		FTransform ComponentTransformWS = FTransform::Identity;
-		FVector ComponentVelocityCS = FVector::ZeroVector;
+		FVector ComponentMoveDeltaWS = FVector::ZeroVector;
+		FVector CharacterVelocityWS = FVector::ZeroVector;
+		FVector SmoothCapsuleGroundNormalWS = FVector::ZeroVector;
+		FQuaternionSpringState SmoothCapsuleGroundNormalSpringState;
 		bool bIsOnGround = false;
 	};
 
@@ -299,12 +303,6 @@ public:
 	UPROPERTY(EditAnywhere, Category = "Trace Settings")
 	float MaxGroundPenetration = 10.0f;
 
-	// How much we align to simple vs complex collision when the foot is in flight
-	// Tracing against simple geometry (i.e. it's common for stairs to have simplified ramp collisions) can provide a 
-	// smoother trajectory when the foot is in flight
-	UPROPERTY(EditAnywhere, Category = "Trace Settings", meta = (EditCondition = "bEnabled", DisplayAfter = "bEnabled"))
-	float SimpleCollisionInfluence = 0.0f;
-
 	// The channel to use for our simple trace
 	UPROPERTY(EditAnywhere, Category = "Trace Settings", meta = (EditCondition = "bEnabled", DisplayAfter = "bEnabled"))
 	TEnumAsByte<ETraceTypeQuery> SimpleTraceChannel = TraceTypeQuery1;
@@ -350,11 +348,11 @@ enum class EPelvisHeightMode : uint8
 UENUM(BlueprintType)
 enum class EActorMovementCompensationMode : uint8
 {
-	// Keep pelvis component-space and follow along all of the actor's vertical ground movement
+	// Keep pelvis component-space and follow along all of the actor's vertical ground movement. Use when you have really smooth ground geometry, or if you have moving platforms.
 	ComponentSpace,
 	// Hold pelvis world-space and ignore the actor's vertical ground movement. Let springs interpolate the difference
 	WorldSpace,
-	// Keep pelvis component-space, but hold world-space transform when the actor does sudden changes (i.e. a big step), and let springs interpolate the difference.
+	// Keep pelvis component-space, but hold world-space transform when the actor does sudden changes (i.e. a big step), and let springs interpolate the difference. Does not work on moving platforms.
 	SuddenMotionOnly
 };
 
@@ -507,6 +505,10 @@ public:
 	UPROPERTY(EditAnywhere, Category = "Plant Settings", meta = (ClampMin = "0.0", UIMin = "0.0", ClampMax = "1.0", UIMax = "1.0"))
 	float AnkleTwistReduction = 0.75f;
 
+	// Only consider world deltas caused by played velocity. This is useful when on a moving platform i.e. but won't consider slopes.
+	UPROPERTY(EditAnywhere, Category = "Plant Settings")
+	bool bReconstructWorldPlantFromVelocity = false;
+
 	// Whether to allow adjusting the heel lift before we plant
 	UPROPERTY(EditAnywhere, Category = "Plant Settings")
 	bool bAdjustHeelBeforePlanting = false;
@@ -549,6 +551,9 @@ public:
 	UPROPERTY(EditAnywhere, BlueprintReadWrite, Category = "Settings", meta = (PinHiddenByDefault))
 	FFootPlacementTraceSettings TraceSettings;
 
+	UPROPERTY(EditAnywhere, Category = Settings, meta = (PinHiddenByDefault))
+	FVector BaseTranslationDelta = FVector::ZeroVector;
+
 public:
 	FAnimNode_FootPlacement();
 
@@ -590,9 +595,8 @@ private:
 
 	// Calculate the desired pelvis offset, based on procedural character/foot adjustments
 	FTransform SolvePelvis(const UE::Anim::FootPlacement::FEvaluationContext& Context);
-	TBitArray<> FindRelevantFeet(const UE::Anim::FootPlacement::FEvaluationContext& Context);
-	TBitArray<> FindPlantedFeet(const UE::Anim::FootPlacement::FEvaluationContext& Context);
-	FTransform UpdatePelvisInterpolation(
+
+	FTransform UpdatePelvisInterpolationRootSpace(
 		const UE::Anim::FootPlacement::FEvaluationContext& Context,
 		const FTransform& TargetPelvisTransform);
 
@@ -605,6 +609,8 @@ private:
 
 	FVector GetApproachDirWS(const FAnimationBaseContext& Context) const;
 
+	const FTransform& GetRootToComponent() const;
+
 private:
 	float CachedDeltaTime = 0.0f;
 	FVector LastComponentLocation = FVector::ZeroVector;
@@ -665,8 +671,7 @@ private:
 	// Handles horizontal interpolation when unlocking the plant
 	FTransform UpdatePlantOffsetInterpolation(
 		const UE::Anim::FootPlacement::FEvaluationContext& Context,
-		UE::Anim::FootPlacement::FLegRuntimeData::FInterpolationData& InOutInterpData,
-		const FTransform& DesiredTransformCS) const;
+		UE::Anim::FootPlacement::FLegRuntimeData::FInterpolationData& InOutInterpData) const;
 
 	// Handles the interpolation of the planting plane. Because the plant transform is specified with respect to the 
 	// planting plane, it cannot change abruptly without causing an animation pop. It must be interpolated instead.
@@ -676,6 +681,7 @@ private:
 		const FTransform& LastAlignedFootTransform,
 		const float AlignmentAlpha,
 		FPlane& InOutPlantPlane,
+		const UE::Anim::FootPlacement::FLegRuntimeData::FInputPoseData& LegInputPose,
 		UE::Anim::FootPlacement::FLegRuntimeData::FInterpolationData& InOutInterpData) const;
 
 	// Checks unplanting and replanting conditions to determine if the foot is planted
