Foot Placement: Clean up plant rotation, fix bad behavior during jumps and added debug viz for smooth capsule ground plane

diff --git a/Engine/Plugins/Animation/AnimationWarping/Source/Runtime/Private/BoneControllers/AnimNode_FootPlacement.cpp b/Engine/Plugins/Animation/AnimationWarping/Source/Runtime/Private/BoneControllers/AnimNode_FootPlacement.cpp
--- a/Engine/Plugins/Animation/AnimationWarping/Source/Runtime/Private/BoneControllers/AnimNode_FootPlacement.cpp
+++ b/Engine/Plugins/Animation/AnimationWarping/Source/Runtime/Private/BoneControllers/AnimNode_FootPlacement.cpp
@@ -510,15 +510,10 @@ void FAnimNode_FootPlacement::UpdatePlantingPlaneInterpolation(
 
 	if (CharacterData.bIsOnGround == false)
 	{
-		// if we're in the air, try to bring foot to source pose and collide with the ground.
-		const FPlane SourceFootPlane = FPlane(FootTransformWS.GetLocation(), Context.ApproachDirCS);
-		const float DistanceToGroundPlane =
-			UE::Anim::FootPlacement::GetDistanceToPlaneAlongDirection(ImpactLocationWS, SourceFootPlane, Context.ApproachDirCS);
-		const float PenetrationAmount = -DistanceToGroundPlane - LegInputPose.DistanceToPlant;
-		if (PenetrationAmount > 0.0f)
-		{
-			InOutPlantPlane = InOutPlantPlane.TranslateBy(-Context.ApproachDirCS * PenetrationAmount);
-		}
+		const FPlane GroundPlane = FPlane(ImpactLocationWS, -Context.ApproachDirWS);
+		const FVector SourceGroundPoint = Context.OwningComponentToWorld.TransformPosition(PelvisData.InputPose.IKRootTransformCS.GetLocation());
+		// if we're in the air, try to bring foot to source pose
+		InOutPlantPlane = FPlane(SourceGroundPoint, -Context.ApproachDirWS);
 	}
 
 	if (InterpolationSettings.bEnableFloorInterpolation && !bIsFirstUpdate)
@@ -543,17 +538,20 @@ void FAnimNode_FootPlacement::UpdatePlantingPlaneInterpolation(
 		float AdjustedPrevZ = FMath::Abs(LastPlaneDeltaZ) < FMath::Abs(PrevPlaneDeltaZ) ?
 			LastPlaneIntersection.Z : PrevPlaneIntersection.Z;
 
-		// Since our ground plane is in root space, apply the component delta if it gets us closer to the desired ground plane.
-		// This means if your ground geometry for capsule is smooth, we want to leverage that for foot alignment. 
-		// When ground geometry is not smooth, we follow the same logic as pelvis interpolation, and attempt to smooth out the jump in world space.
-		const float GroundPlaneDelta = CurrPlaneIntersection.Z  - AdjustedPrevZ;
-		if ((GroundPlaneDelta) > 0.0f)
-		{
-			AdjustedPrevZ += FMath::Clamp(-CharacterData.ComponentMoveDeltaWS.Z, 0.0f, GroundPlaneDelta);
-		}
-		else
+		if (CharacterData.bIsOnGround)
 		{
-			AdjustedPrevZ += FMath::Clamp(-CharacterData.ComponentMoveDeltaWS.Z, GroundPlaneDelta, 0.0f);
+			// Since our ground plane is in root space, apply the component delta if it gets us closer to the desired ground plane.
+			// This means if your ground geometry for capsule is smooth, we want to leverage that for foot alignment. 
+			// When ground geometry is not smooth, we follow the same logic as pelvis interpolation, and attempt to smooth out the jump in world space.
+			const float GroundPlaneDelta = CurrPlaneIntersection.Z  - AdjustedPrevZ;
+			if ((GroundPlaneDelta) > 0.0f)
+			{
+				AdjustedPrevZ += FMath::Min(FMath::Abs(-CharacterData.ComponentMoveDeltaWS.Z), GroundPlaneDelta);
+			}
+			else
+			{
+				AdjustedPrevZ += FMath::Max(-FMath::Abs(CharacterData.ComponentMoveDeltaWS.Z), GroundPlaneDelta);
+			}
 		}
 
 		float PlantPlaneSpringHeight = UKismetMathLibrary::FloatSpringInterp(
@@ -857,8 +855,9 @@ UE::Anim::FootPlacement::FPlantResult FAnimNode_FootPlacement::FinalizeFootAlign
 		CorrectedFootTransformCS.GetLocation(),
 		PlantPlaneCS,
 		Context.ApproachDirCS);
-	const float MinDistance = FMath::Min(FootDistance, BallDistance);
-
+	float MinDistance = FMath::Min(FootDistance, BallDistance);
+	// Allow as much penetration as the source animation.
+	MinDistance -= FMath::Min(0.0f, LegData.InputPose.DistanceToPlant);
 	// A min distance < 0.0f means there was penetration
 	if (MinDistance < 0.0f)
 	{
@@ -1331,6 +1330,11 @@ void FAnimNode_FootPlacement::EvaluateSkeletalControl_AnyThread(FComponentSpaceP
 			// Draw pelvis interpolation target
 			AnimInstanceProxy->AnimDrawDebugPoint(
 				PelvisTargetTransformWS.GetLocation(), 10.0f, FColor::Purple, false, -1.0f, SDPG_Foreground);
+
+			const FVector IKFootRootLocationWS = ComponentTransform.TransformPosition(PelvisData.InputPose.IKRootTransformCS.GetLocation());
+			AnimInstanceProxy->AnimDrawDebugCircle(
+				IKFootRootLocationWS, 100.0f, 24, FColor::Cyan,
+				CharacterData.SmoothCapsuleGroundNormalWS, false, -1.0f, SDPG_Foreground, 0.5f);
 		}
 		
 #if ENABLE_FOOTPLACEMENT_DEBUG
@@ -1345,6 +1349,9 @@ void FAnimNode_FootPlacement::EvaluateSkeletalControl_AnyThread(FComponentSpaceP
 			UE_VLOG_SPHERE(AnimInstance, "FootPlacement", Display, PelvisTransformWS.GetTranslation(), 0, FColor::Green, TEXT(""));
 			UE_VLOG_SPHERE(AnimInstance, "FootPlacement", Display, BasePelvisTransformWS.GetTranslation(), 0, FColor::Blue, TEXT(""));
 			UE_VLOG_SPHERE(AnimInstance, "FootPlacement", Display, PelvisTargetTransformWS.GetTranslation(), 0, FColor::Purple, TEXT(""));
+
+			const FVector IKFootRootLocationWS = ComponentTransform.TransformPosition(PelvisData.InputPose.IKRootTransformCS.GetLocation());
+			UE_VLOG_CIRCLE_THICK(AnimInstance, "FootPlacement", Display, IKFootRootLocationWS, CharacterData.SmoothCapsuleGroundNormalWS, 50, FColor::Cyan, 1, TEXT(""));
 		}
 #endif
 	}
@@ -1606,13 +1613,21 @@ void FAnimNode_FootPlacement::ProcessCharacterState(const UE::Anim::FootPlacemen
 		const FVector CapsuleMoveOffsetWS =
 			(ComponentLocationWS - OwningComponentAdjustedLastLocationWS - BaseTranslationDelta) * -Context.ApproachDirWS;
 
-		CharacterData.ComponentMoveDeltaWS +=CapsuleMoveOffsetWS;
+		CharacterData.ComponentMoveDeltaWS -= CapsuleMoveOffsetWS;
 		if (!CapsuleMoveOffsetWS.IsNearlyZero(KINDA_SMALL_NUMBER))
 		{
 			const FVector CapsuleMoveOffsetCS =
 				Context.OwningComponentToWorld.InverseTransformVectorNoScale(CapsuleMoveOffsetWS);
+
 			// Offseting our interpolator lets it smoothly solve sudden capsule deltas, instead of following it and pop
 			PelvisData.Interpolation.PelvisTranslationOffset -= CapsuleMoveOffsetCS;
+
+			for (UE::Anim::FootPlacement::FLegRuntimeData& LegData : LegsData)
+			{
+				// Also offset our foot plant plane interpolators by this same delta.
+				//@TODO: this should be in root space too, but vertical motion is the same regardless of space for now.
+				LegData.Plant.PlantPlaneRS = LegData.Plant.PlantPlaneRS.TranslateBy(-CapsuleMoveOffsetCS);
+			}
 		}
 	}
 
@@ -1845,7 +1860,7 @@ void FAnimNode_FootPlacement::ProcessFootAlignment(
 	UpdatePlantingPlaneInterpolation(Context, LegData.UnalignedFootTransformWS,
 									LastAlignedFootTransformWS, 
 									InputPose.AlignmentAlpha, 
-		PlantPlaneWS, 
+									PlantPlaneWS, 
 									InputPose, 
 									Interpolation);
 	const FPlane PlantPlaneCS = PlantPlaneWS.TransformBy(ComponentToWorldInv.ToMatrixWithScale());
@@ -1857,19 +1872,8 @@ void FAnimNode_FootPlacement::ProcessFootAlignment(
 	AlignPlantToGround(Context, PlantPlaneWS, InputPose, LegData.AlignedFootTransformWS, Plant.TwistCorrection);
 
 	const FTransform AlignedFootTransformCS =	LegData.AlignedFootTransformWS * ComponentToWorldInv;
-	// The target transform is a blend based on FkAlignmentAlpha.
-	// Until we have prediction, favor the ground aligned position, 
-	// since this will likely have a more accurate distance from plane
-	FTransform BlendedPlantTransformCS = AlignedFootTransformCS;
-	// When unplanted/unaligned, favor FK orientation and fix penetrations later. 
-	BlendedPlantTransformCS.SetRotation(
-		FQuat::Slerp(
-			InputPose.FootTransformCS.GetRotation(),
-			AlignedFootTransformCS.GetRotation(),
-			InputPose.AlignmentAlpha));
-
-	LegData.AlignedFootTransformRS = BlendedPlantTransformCS.GetRelativeTransform(RootToComponent);
-	LegData.AlignedFootTransformWS = BlendedPlantTransformCS * Context.OwningComponentToWorld;
+	LegData.AlignedFootTransformRS = AlignedFootTransformCS.GetRelativeTransform(RootToComponent);
+	LegData.AlignedFootTransformWS = AlignedFootTransformCS * Context.OwningComponentToWorld;
 }
 
 FVector FAnimNode_FootPlacement::GetApproachDirWS(const FAnimationBaseContext& Context) const
