Fix Orientation Warping Counter Compensate undoing root motion rotation.

#tests PIE

diff --git a/Engine/Plugins/Animation/AnimationWarping/Source/Runtime/Private/BoneControllers/AnimNode_OrientationWarping.cpp b/Engine/Plugins/Animation/AnimationWarping/Source/Runtime/Private/BoneControllers/AnimNode_OrientationWarping.cpp
--- a/Engine/Plugins/Animation/AnimationWarping/Source/Runtime/Private/BoneControllers/AnimNode_OrientationWarping.cpp
+++ b/Engine/Plugins/Animation/AnimationWarping/Source/Runtime/Private/BoneControllers/AnimNode_OrientationWarping.cpp
@@ -226,7 +226,9 @@ void FAnimNode_OrientationWarping::EvaluateSkeletalControl_AnyThread(FComponentS
 			}
 
 			FVector RootMotionDeltaTranslation = RootMotionTransformDelta.GetTranslation();
-
+			const FQuat PreviousRootMotionDeltaRotation = RootMotionDeltaRotation;
+			RootMotionDeltaRotation = RootMotionTransformDelta.GetRotation();
+			
 			// Flatten root motion translation, along the rotation axis.
 			RootMotionDeltaTranslation = RootMotionDeltaTranslation - RotationAxisVector.Dot(RootMotionDeltaTranslation) * RotationAxisVector;
 
@@ -296,7 +298,7 @@ void FAnimNode_OrientationWarping::EvaluateSkeletalControl_AnyThread(FComponentS
 					}
 				}
 
-				// Don't compensate interpolation by the root motion angle delta if the previous angle isn't valid.
+				// Don't compensate interpolation by the root motion angle delta if the previous direction isn't valid.
 				if (bCounterCompenstateInterpolationByRootMotion && !PreviousRootMotionDeltaDirection.IsNearlyZero(UE_SMALL_NUMBER))
 				{
 #if !ENABLE_ANIM_DEBUG
@@ -304,7 +306,10 @@ void FAnimNode_OrientationWarping::EvaluateSkeletalControl_AnyThread(FComponentS
 #endif
 					// Counter the interpolated orientation angle by the root motion direction angle delta.
 					// This prevents our interpolation from fighting the natural root motion that's flowing through the graph.
-					RootMotionDeltaAngleRad = UE::Anim::SignedAngleRadBetweenNormals(RootMotionDeltaDirection, PreviousRootMotionDeltaDirection, RotationAxisVector);
+					// To correctly measure the amount to counter, we need to unrotate our previous delta direction by our previous rotation
+					// As the previous direction delta is relative to the previous rotation delta
+					RootMotionDeltaAngleRad = UE::Anim::SignedAngleRadBetweenNormals(RootMotionDeltaDirection, PreviousRootMotionDeltaRotation.UnrotateVector(PreviousRootMotionDeltaDirection), RotationAxisVector);
+					
 					// Root motion may have large deltas i.e. bad blends or sudden direction changes like pivots.
 					// If there's an instantaneous pop in root motion direction, this is likely a pivot.
 					const float MaxRootMotionDeltaToCompensateRad = FMath::DegreesToRadians(MaxRootMotionDeltaToCompensateDegrees);
diff --git a/Engine/Plugins/Animation/AnimationWarping/Source/Runtime/Public/BoneControllers/AnimNode_OrientationWarping.h b/Engine/Plugins/Animation/AnimationWarping/Source/Runtime/Public/BoneControllers/AnimNode_OrientationWarping.h
--- a/Engine/Plugins/Animation/AnimationWarping/Source/Runtime/Public/BoneControllers/AnimNode_OrientationWarping.h
+++ b/Engine/Plugins/Animation/AnimationWarping/Source/Runtime/Public/BoneControllers/AnimNode_OrientationWarping.h
@@ -221,6 +221,7 @@ private:
 	
 	// Internal current frame root motion delta direction
 	FVector RootMotionDeltaDirection = FVector::ZeroVector;
+	FQuat RootMotionDeltaRotation = FQuat::Identity;
 
 	// Internal orientation warping angle
 	float ActualOrientationAngleRad = 0.f;
