Prevent Foot Placement running when an actor is dropped into the level viewport

diff --git a/Engine/Plugins/Animation/AnimationWarping/Source/Runtime/Private/BoneControllers/AnimNode_FootPlacement.cpp b/Engine/Plugins/Animation/AnimationWarping/Source/Runtime/Private/BoneControllers/AnimNode_FootPlacement.cpp
--- a/Engine/Plugins/Animation/AnimationWarping/Source/Runtime/Private/BoneControllers/AnimNode_FootPlacement.cpp
+++ b/Engine/Plugins/Animation/AnimationWarping/Source/Runtime/Private/BoneControllers/AnimNode_FootPlacement.cpp
@@ -1159,6 +1159,16 @@ void FAnimNode_FootPlacement::EvaluateSkeletalControl_AnyThread(FComponentSpaceP
 
 	check(OutBoneTransforms.Num() == 0);
 
+#if WITH_EDITOR
+	// Early out if we are in an editor world and not updating animation to avoid generating invalid data
+	USkeletalMeshComponent* SMC = Output.AnimInstanceProxy->GetSkelMeshComponent();
+	bool bIsGameOrEditorPreviewWorld = SMC->GetWorld()->IsGameWorld() || SMC->GetWorld()->WorldType == EWorldType::EditorPreview;
+	if (!bIsGameOrEditorPreviewWorld && !SMC->GetUpdateAnimationInEditor())
+	{
+		return;
+	}
+#endif
+
 	// Manually calculate distance instead of using the teleport flag to properly handle cases like crouch
 	// and instantaneous root offsets i.e. when entering/leaving a vehicle.
 	// See FAnimNode_Inertialization::Evaluate_AnyThread and/or UE-78594
