Improve debugging on OrientationWarp
- Save prediction info so prediction can be analyzed during pauses
- Add to visual logging

#tests PIE

diff --git a/Engine/Plugins/Animation/AnimationWarping/Source/Runtime/Private/BoneControllers/AnimNode_OrientationWarping.cpp b/Engine/Plugins/Animation/AnimationWarping/Source/Runtime/Private/BoneControllers/AnimNode_OrientationWarping.cpp
--- a/Engine/Plugins/Animation/AnimationWarping/Source/Runtime/Private/BoneControllers/AnimNode_OrientationWarping.cpp
+++ b/Engine/Plugins/Animation/AnimationWarping/Source/Runtime/Private/BoneControllers/AnimNode_OrientationWarping.cpp
@@ -151,8 +151,13 @@ void FAnimNode_OrientationWarping::EvaluateSkeletalControl_AnyThread(FComponentS
 	FTransform RootMotionTransformDelta = FTransform::Identity;
 	float RootMotionDeltaAngleRad = 0.0;
 	const float PreviousOrientationAngleRad = ActualOrientationAngleRad;
-	FVector CurrentRootMotionDeltaDirection = FVector::ZeroVector;
-	bool bUsedFutureRootMotion = false;
+#endif
+
+#if ENABLE_ANIM_DEBUG || ENABLE_VISUAL_LOG
+	if (DeltaSeconds > 0)
+	{
+		bUsedFutureRootMotion = false;
+	}
 #endif
 
 	// We will likely need to revisit LocomotionAngle participating as an input to orientation warping.
@@ -288,7 +293,7 @@ void FAnimNode_OrientationWarping::EvaluateSkeletalControl_AnyThread(FComponentS
 						// Check for continuity. Later, when the future and current are re-aligned, error will be equivalent, so we can fall back to non-future root motion
 						if (FMath::Abs(PredictedToPrevErrorAngleRad) < FMath::Abs(CurrentToPrevErrorAngleRad) + FMath::DegreesToRadians(TargetBiasAngle))
 						{
-#if ENABLE_ANIM_DEBUG
+#if ENABLE_ANIM_DEBUG || ENABLE_VISUAL_LOG
 							CurrentRootMotionDeltaDirection = RootMotionDeltaDirection;
 							bUsedFutureRootMotion = true;
 #endif
@@ -410,7 +415,7 @@ void FAnimNode_OrientationWarping::EvaluateSkeletalControl_AnyThread(FComponentS
 				Output.AnimInstanceProxy->AnimDrawDebugDirectionalArrow(
 					ComponentTransform.GetLocation() + DebugArrowOffset,
 					ComponentTransform.GetLocation() + DebugArrowOffset + RotationDirection * 100.f * DebugDrawScale,
-					40.f * DebugDrawScale, FColor::Cyan, false, 0.f, 2.f * DebugDrawScale);
+					40.f * DebugDrawScale, FColor::Yellow, false, 0.f, 2.f * DebugDrawScale);
 			}
 			else
 			{
@@ -511,10 +516,28 @@ void FAnimNode_OrientationWarping::EvaluateSkeletalControl_AnyThread(FComponentS
 				: ActorForwardDirection.RotateAngleAxis(OrientationAngle, RotationAxisVector);
 
 			DebugArrowOffset += FVector::ZAxisVector * DebugDrawScale;
-			UE_VLOG_ARROW(Output.AnimInstanceProxy->GetAnimInstanceObject(), "OrientationWarping", Display,
-				ComponentTransform.GetLocation() + DebugArrowOffset,
-				ComponentTransform.GetLocation() + DebugArrowOffset + RotationDirection * 100.f * DebugDrawScale,
-				FColor::Blue, TEXT(""));
+
+			if (bUsedFutureRootMotion && bGraphDrivenWarping)
+			{
+				const FVector CurrentRotationDirection = ComponentTransform.GetRotation().RotateVector(CurrentRootMotionDeltaDirection);
+
+				UE_VLOG_ARROW(Output.AnimInstanceProxy->GetAnimInstanceObject(), "OrientationWarping", Display,
+					ComponentTransform.GetLocation() + DebugArrowOffset,
+					ComponentTransform.GetLocation() + DebugArrowOffset + CurrentRotationDirection * 100.f * DebugDrawScale,
+					FColor::Blue, TEXT(""));
+
+				UE_VLOG_ARROW(Output.AnimInstanceProxy->GetAnimInstanceObject(), "OrientationWarping", Display,
+					ComponentTransform.GetLocation() + DebugArrowOffset,
+					ComponentTransform.GetLocation() + DebugArrowOffset + RotationDirection * 100.f * DebugDrawScale,
+					FColor::Yellow, TEXT(""));
+			}
+			else
+			{
+				UE_VLOG_ARROW(Output.AnimInstanceProxy->GetAnimInstanceObject(), "OrientationWarping", Display,
+					ComponentTransform.GetLocation() + DebugArrowOffset,
+					ComponentTransform.GetLocation() + DebugArrowOffset + RotationDirection * 100.f * DebugDrawScale,
+					FColor::Blue, TEXT(""));
+			}
 
 			const float ActualOrientationAngleDegrees = FMath::RadiansToDegrees(ActualOrientationAngleRad);
 			const FVector WarpedRotationDirection = bGraphDrivenWarping
diff --git a/Engine/Plugins/Animation/AnimationWarping/Source/Runtime/Public/BoneControllers/AnimNode_OrientationWarping.h b/Engine/Plugins/Animation/AnimationWarping/Source/Runtime/Public/BoneControllers/AnimNode_OrientationWarping.h
--- a/Engine/Plugins/Animation/AnimationWarping/Source/Runtime/Public/BoneControllers/AnimNode_OrientationWarping.h
+++ b/Engine/Plugins/Animation/AnimationWarping/Source/Runtime/Public/BoneControllers/AnimNode_OrientationWarping.h
@@ -223,6 +223,12 @@ private:
 	FVector RootMotionDeltaDirection = FVector::ZeroVector;
 	FQuat RootMotionDeltaRotation = FQuat::Identity;
 
+#if ENABLE_ANIM_DEBUG || ENABLE_VISUAL_LOG
+	// Store these in debug only so that they can persist in debug / across pauses
+	FVector CurrentRootMotionDeltaDirection = FVector::ZeroVector;
+	bool bUsedFutureRootMotion = false;
+#endif
+
 	// Internal orientation warping angle
 	float ActualOrientationAngleRad = 0.f;
 	float BlendWeight = 0.0f;
