Fix Orientation Warping jittering with counter compensate
- Don't update root motion direction based on prediction.
- Remove threshholds / checks needed for contunity due to innapropriate root motion direction setting

#tests PIE

diff --git a/Engine/Plugins/Animation/AnimationWarping/Source/Runtime/Private/BoneControllers/AnimNode_OrientationWarping.cpp b/Engine/Plugins/Animation/AnimationWarping/Source/Runtime/Private/BoneControllers/AnimNode_OrientationWarping.cpp
--- a/Engine/Plugins/Animation/AnimationWarping/Source/Runtime/Private/BoneControllers/AnimNode_OrientationWarping.cpp
+++ b/Engine/Plugins/Animation/AnimationWarping/Source/Runtime/Private/BoneControllers/AnimNode_OrientationWarping.cpp
@@ -266,7 +266,6 @@ void FAnimNode_OrientationWarping::EvaluateSkeletalControl_AnyThread(FComponentS
 				// When that happens, since we use the delta between root motion and movement direction, we would be over-rotating the lower body and breaking the pose during those frames
 				// So, when that happens we use the inverse of the root motion direction to calculate our target rotation. 
 				// This feels a bit 'hacky' but its the only option I've found so far to mitigate the problem
-				float PreFlipOrientationAngleRad = TargetOrientationAngleRad;
 				if (LocomotionAngleDeltaThreshold > 0.f)
 				{
 					if (FMath::Abs(FMath::RadiansToDegrees(TargetOrientationAngleRad)) > LocomotionAngleDeltaThreshold)
@@ -283,23 +282,16 @@ void FAnimNode_OrientationWarping::EvaluateSkeletalControl_AnyThread(FComponentS
 					float PredictedOrientationErrorAngleRad = UE::Anim::SignedAngleRadBetweenNormals(PredictedRootMotionDeltaTranslation, LocomotionForward, RotationAxisVector);
 
 					// The future orientation will often match the current, so add a small delta to avoid further testing for same values
-					if (FMath::Abs(PredictedOrientationErrorAngleRad) + UE_KINDA_SMALL_NUMBER < FMath::Abs(PreFlipOrientationAngleRad))
+					if (FMath::Abs(PredictedOrientationErrorAngleRad) + UE_KINDA_SMALL_NUMBER < FMath::Abs(TargetOrientationAngleRad))
 					{
-						// Error is normalized. So even if the future root motion orientation has less error, it may be in the opposite angle.
-						// Thus only use the future root motion for warping if it also has less error relative to the previous position
-						float PredictedToPrevErrorAngleRad = UE::Anim::SignedAngleRadBetweenNormals(PreviousRootMotionDeltaDirection, PredictedRootMotionDeltaTranslation, RotationAxisVector);
-						float CurrentToPrevErrorAngleRad = UE::Anim::SignedAngleRadBetweenNormals(PreviousRootMotionDeltaDirection, RootMotionDeltaDirection, RotationAxisVector);
-
-						// Check for continuity. Later, when the future and current are re-aligned, error will be equivalent, so we can fall back to non-future root motion
-						if (FMath::Abs(PredictedToPrevErrorAngleRad) < FMath::Abs(CurrentToPrevErrorAngleRad) + FMath::DegreesToRadians(TargetBiasAngle))
-						{
+						// Note: Don't update root motion direction, as the root motion direction is what we are playing, 
+						// not the future. Updating root motion direction will break counter compensate. 
+						// Also, we rely on the built in interp for continuity / smoothness
 #if ENABLE_ANIM_DEBUG || ENABLE_VISUAL_LOG
-							CurrentRootMotionDeltaDirection = RootMotionDeltaDirection;
-							bUsedFutureRootMotion = true;
+						FutureRootMotionDeltaDirection = PredictedRootMotionDeltaTranslation;
+						bUsedFutureRootMotion = true;
 #endif
-							TargetOrientationAngleRad = PredictedOrientationErrorAngleRad;
-							RootMotionDeltaDirection = PredictedRootMotionDeltaTranslation;
-						}
+						TargetOrientationAngleRad = PredictedOrientationErrorAngleRad;
 					}
 				}
 
@@ -402,28 +394,20 @@ void FAnimNode_OrientationWarping::EvaluateSkeletalControl_AnyThread(FComponentS
 				: ActorForwardDirection.RotateAngleAxis(OrientationAngle, RotationAxisVector);
 
 			DebugArrowOffset += FVector::ZAxisVector * DebugDrawScale;
+			Output.AnimInstanceProxy->AnimDrawDebugDirectionalArrow(
+				ComponentTransform.GetLocation() + DebugArrowOffset,
+				ComponentTransform.GetLocation() + DebugArrowOffset + RotationDirection * 100.f * DebugDrawScale,
+				40.f * DebugDrawScale, FColor::Blue, false, 0.f, 2.f * DebugDrawScale);
 
 			if (bUsedFutureRootMotion && bGraphDrivenWarping)
 			{
-				const FVector CurrentRotationDirection = ComponentTransform.GetRotation().RotateVector(CurrentRootMotionDeltaDirection);
-
-				Output.AnimInstanceProxy->AnimDrawDebugDirectionalArrow(
-					ComponentTransform.GetLocation() + DebugArrowOffset,
-					ComponentTransform.GetLocation() + DebugArrowOffset + CurrentRotationDirection * 100.f * DebugDrawScale,
-					40.f * DebugDrawScale, FColor::Blue, false, 0.f, 2.f * DebugDrawScale);
+				const FVector FutureRotationDirection = ComponentTransform.GetRotation().RotateVector(FutureRootMotionDeltaDirection);
 
 				Output.AnimInstanceProxy->AnimDrawDebugDirectionalArrow(
 					ComponentTransform.GetLocation() + DebugArrowOffset,
-					ComponentTransform.GetLocation() + DebugArrowOffset + RotationDirection * 100.f * DebugDrawScale,
+					ComponentTransform.GetLocation() + DebugArrowOffset + FutureRotationDirection * 100.f * DebugDrawScale,
 					40.f * DebugDrawScale, FColor::Yellow, false, 0.f, 2.f * DebugDrawScale);
 			}
-			else
-			{
-				Output.AnimInstanceProxy->AnimDrawDebugDirectionalArrow(
-					ComponentTransform.GetLocation() + DebugArrowOffset,
-					ComponentTransform.GetLocation() + DebugArrowOffset + RotationDirection * 100.f * DebugDrawScale,
-					40.f * DebugDrawScale, FColor::Blue, false, 0.f, 2.f * DebugDrawScale);
-			}
 
 			const float ActualOrientationAngleDegrees = FMath::RadiansToDegrees(ActualOrientationAngleRad);
 			const FVector WarpedRotationDirection = bGraphDrivenWarping
@@ -516,27 +500,19 @@ void FAnimNode_OrientationWarping::EvaluateSkeletalControl_AnyThread(FComponentS
 				: ActorForwardDirection.RotateAngleAxis(OrientationAngle, RotationAxisVector);
 
 			DebugArrowOffset += FVector::ZAxisVector * DebugDrawScale;
+			UE_VLOG_ARROW(Output.AnimInstanceProxy->GetAnimInstanceObject(), "OrientationWarping", Display,
+				ComponentTransform.GetLocation() + DebugArrowOffset,
+				ComponentTransform.GetLocation() + DebugArrowOffset + RotationDirection * 100.f * DebugDrawScale,
+				FColor::Blue, TEXT(""));
 
 			if (bUsedFutureRootMotion && bGraphDrivenWarping)
 			{
-				const FVector CurrentRotationDirection = ComponentTransform.GetRotation().RotateVector(CurrentRootMotionDeltaDirection);
+				const FVector FutureRotationDirection = ComponentTransform.GetRotation().RotateVector(FutureRootMotionDeltaDirection);
 
-				UE_VLOG_ARROW(Output.AnimInstanceProxy->GetAnimInstanceObject(), "OrientationWarping", Display,
-					ComponentTransform.GetLocation() + DebugArrowOffset,
-					ComponentTransform.GetLocation() + DebugArrowOffset + CurrentRotationDirection * 100.f * DebugDrawScale,
-					FColor::Blue, TEXT(""));
-
-				UE_VLOG_ARROW(Output.AnimInstanceProxy->GetAnimInstanceObject(), "OrientationWarping", Display,
-					ComponentTransform.GetLocation() + DebugArrowOffset,
-					ComponentTransform.GetLocation() + DebugArrowOffset + RotationDirection * 100.f * DebugDrawScale,
-					FColor::Yellow, TEXT(""));
-			}
-			else
-			{
-				UE_VLOG_ARROW(Output.AnimInstanceProxy->GetAnimInstanceObject(), "OrientationWarping", Display,
+				Output.AnimInstanceProxy->AnimDrawDebugDirectionalArrow(
 					ComponentTransform.GetLocation() + DebugArrowOffset,
-					ComponentTransform.GetLocation() + DebugArrowOffset + RotationDirection * 100.f * DebugDrawScale,
-					FColor::Blue, TEXT(""));
+					ComponentTransform.GetLocation() + DebugArrowOffset + FutureRotationDirection * 100.f * DebugDrawScale,
+					40.f * DebugDrawScale, FColor::Yellow, false, 0.f, 2.f * DebugDrawScale);
 			}
 
 			const float ActualOrientationAngleDegrees = FMath::RadiansToDegrees(ActualOrientationAngleRad);
diff --git a/Engine/Plugins/Animation/AnimationWarping/Source/Runtime/Public/BoneControllers/AnimNode_OrientationWarping.h b/Engine/Plugins/Animation/AnimationWarping/Source/Runtime/Public/BoneControllers/AnimNode_OrientationWarping.h
--- a/Engine/Plugins/Animation/AnimationWarping/Source/Runtime/Public/BoneControllers/AnimNode_OrientationWarping.h
+++ b/Engine/Plugins/Animation/AnimationWarping/Source/Runtime/Public/BoneControllers/AnimNode_OrientationWarping.h
@@ -41,13 +41,6 @@ struct ANIMATIONWARPINGRUNTIME_API FAnimNode_OrientationWarping : public FAnimNo
 	UPROPERTY(EditAnywhere, BlueprintReadWrite, Category = Experimental, meta = (PinShownByDefault))
 	float TargetTime = 0.8f;
 
-	// Experimental. A bias in which the future root motion is preferred, even if the current root motion is closer to the previous orientation.
-	// We compare the current / future root motion to our previous orientation for continuity. 
-	// By default the current root motion is used, so it will always win this comparision without some bias.
-	// Note: We default to 5 degrees which may seem large, but the animation itself being continuous should prevent popping.
-	UPROPERTY(EditAnywhere, BlueprintReadWrite, Category = Experimental, meta = (PinHiddenByDefault))
-	float TargetBiasAngle = 5.f;
-
 	// The desired orientation angle (in degrees) to warp by relative to the specified RotationAxis
 	UPROPERTY(EditAnywhere, BlueprintReadWrite, Category=Evaluation, meta=(PinShownByDefault))
 	float OrientationAngle = 0.f;
@@ -225,7 +218,7 @@ private:
 
 #if ENABLE_ANIM_DEBUG || ENABLE_VISUAL_LOG
 	// Store these in debug only so that they can persist in debug / across pauses
-	FVector CurrentRootMotionDeltaDirection = FVector::ZeroVector;
+	FVector FutureRootMotionDeltaDirection = FVector::ZeroVector;
 	bool bUsedFutureRootMotion = false;
 #endif
 
