Add a strong interp to Counter Compensate to address stuttering when correcting between two samples between the same frames

#tests PIE, BP

diff --git a/Engine/Plugins/Animation/AnimationWarping/Source/Runtime/Private/BoneControllers/AnimNode_OrientationWarping.cpp b/Engine/Plugins/Animation/AnimationWarping/Source/Runtime/Private/BoneControllers/AnimNode_OrientationWarping.cpp
--- a/Engine/Plugins/Animation/AnimationWarping/Source/Runtime/Private/BoneControllers/AnimNode_OrientationWarping.cpp
+++ b/Engine/Plugins/Animation/AnimationWarping/Source/Runtime/Private/BoneControllers/AnimNode_OrientationWarping.cpp
@@ -312,7 +312,10 @@ void FAnimNode_OrientationWarping::EvaluateSkeletalControl_AnyThread(FComponentS
 					const float MaxRootMotionDeltaToCompensateRad = FMath::DegreesToRadians(MaxRootMotionDeltaToCompensateDegrees);
 					if (FMath::Abs(RootMotionDeltaAngleRad) < MaxRootMotionDeltaToCompensateRad)
 					{
-						ActualOrientationAngleRad = FMath::UnwindRadians(ActualOrientationAngleRad + RootMotionDeltaAngleRad);
+						CounterCompensateTargetAngleRad += RootMotionDeltaAngleRad;
+						float CounterCompensateAngle = FMath::FInterpTo(0, CounterCompensateTargetAngleRad, DeltaSeconds, CounterCompensateInterpSpeed);
+						ActualOrientationAngleRad = FMath::UnwindRadians(ActualOrientationAngleRad + CounterCompensateAngle);
+						CounterCompensateTargetAngleRad -= CounterCompensateAngle;
 					}
 				}
 
diff --git a/Engine/Plugins/Animation/AnimationWarping/Source/Runtime/Public/BoneControllers/AnimNode_OrientationWarping.h b/Engine/Plugins/Animation/AnimationWarping/Source/Runtime/Public/BoneControllers/AnimNode_OrientationWarping.h
--- a/Engine/Plugins/Animation/AnimationWarping/Source/Runtime/Public/BoneControllers/AnimNode_OrientationWarping.h
+++ b/Engine/Plugins/Animation/AnimationWarping/Source/Runtime/Public/BoneControllers/AnimNode_OrientationWarping.h
@@ -110,6 +110,11 @@ struct ANIMATIONWARPINGRUNTIME_API FAnimNode_OrientationWarping : public FAnimNo
 	UPROPERTY(EditAnywhere, Category=Settings, meta=(ClampMin="0.0"))
 	float RotationInterpSpeed = 10.f;
 
+	// Same as RotationInterpSpeed, but for CounterCompensate smoothing. A value of 0 sample raw root motion.
+	// Used to avoid stuttering from resampling root deltas. Root motion is already smooth, so a large value is our default (~75% of 60 fps).
+	UPROPERTY(EditAnywhere, Category = Settings, meta = (ClampMin = "0.0"))
+	float CounterCompensateInterpSpeed = 45.f;
+
 	// Max correction we're allowed to do per-second when using interpolation.
 	// This minimizes pops when we have a large difference between current and target orientation.
 	UPROPERTY(EditAnywhere, Category=Settings, meta=(ClampMin="0.0", EditCondition="RotationInterpSpeed > 0.0f"))
@@ -215,8 +220,13 @@ private:
 	
 	// Internal current frame root motion delta direction
 	FVector RootMotionDeltaDirection = FVector::ZeroVector;
+
+	// Internal current frame root motion delta angle
 	FQuat RootMotionDeltaRotation = FQuat::Identity;
 
+	// Target for counter compenstate, we keep the target so we can smoothly interp.
+	float CounterCompensateTargetAngleRad = 0.0f;
+
 #if ENABLE_ANIM_DEBUG || ENABLE_VISUAL_LOG
 	// Store these in debug only so that they can persist in debug / across pauses
 	FVector FutureRootMotionDeltaDirection = FVector::ZeroVector;
