Secondary debugging improvement to Orientation Warping
- Fade away debug lines via blend weight
- Add a small circle at end of debug lines that corresponds to rewind debugger anim color

#tests Editor, PIE

diff --git a/Engine/Plugins/Animation/AnimationWarping/Source/Runtime/Private/BoneControllers/AnimNode_OrientationWarping.cpp b/Engine/Plugins/Animation/AnimationWarping/Source/Runtime/Private/BoneControllers/AnimNode_OrientationWarping.cpp
--- a/Engine/Plugins/Animation/AnimationWarping/Source/Runtime/Private/BoneControllers/AnimNode_OrientationWarping.cpp
+++ b/Engine/Plugins/Animation/AnimationWarping/Source/Runtime/Private/BoneControllers/AnimNode_OrientationWarping.cpp
@@ -2,6 +2,7 @@
 
 #include "BoneControllers/AnimNode_OrientationWarping.h"
 
+#include "Animation/AnimBlendDebugScope.h"
 #include "Animation/AnimInstanceProxy.h"
 #include "Animation/AnimNodeFunctionRef.h"
 #include "Animation/AnimRootMotionProvider.h"
@@ -16,9 +17,19 @@
 DECLARE_CYCLE_STAT(TEXT("OrientationWarping Eval"), STAT_OrientationWarping_Eval, STATGROUP_Anim);
 
 #if ENABLE_ANIM_DEBUG
-static TAutoConsoleVariable<int32> CVarAnimNodeOrientationWarpingDebug(TEXT("a.AnimNode.OrientationWarping.Debug"), 0, TEXT("Turn on visualization debugging for Orientation Warping."));
-static TAutoConsoleVariable<int32> CVarAnimNodeOrientationWarpingVerbose(TEXT("a.AnimNode.OrientationWarping.Verbose"), 0, TEXT("Turn on verbose graph debugging for Orientation Warping"));
-static TAutoConsoleVariable<int32> CVarAnimNodeOrientationWarpingEnable(TEXT("a.AnimNode.OrientationWarping.Enable"), 1, TEXT("Toggle Orientation Warping"));
+static int32 AnimNodeOrientationWarpingDebug = 0;
+static FAutoConsoleVariableRef CVarAnimNodeOrientationWarpingDebug(TEXT("a.AnimNode.OrientationWarping.Debug"), AnimNodeOrientationWarpingDebug, TEXT("Turn on visualization debugging for Orientation Warping."));
+
+static int32 AnimNodeOrientationWarpingVerbose = 0;
+static FAutoConsoleVariableRef CVarAnimNodeOrientationWarpingVerbose(TEXT("a.AnimNode.OrientationWarping.Verbose"), AnimNodeOrientationWarpingVerbose, TEXT("Turn on verbose graph debugging for Orientation Warping"));
+
+static int32 AnimNodeOrientationWarpingEnable = 1;
+static FAutoConsoleVariableRef CVarAnimNodeOrientationWarpingEnable(TEXT("a.AnimNode.OrientationWarping.Enable"), AnimNodeOrientationWarpingEnable, TEXT("Toggle Orientation Warping"));
+#endif
+
+#if ENABLE_ANIM_DEBUG || ENABLE_VISUAL_LOG
+static int32 AnimNodeOrientationWarpingDebugTransparency = 1;
+static FAutoConsoleVariableRef CVarAnimNodeOrientationWarpingDebugTransparency(TEXT("a.AnimNode.OrientationWarping.Debug.Transparency"), AnimNodeOrientationWarpingDebugTransparency, TEXT("Enable transparency for debug draws, based on blend weight"));
 #endif
 
 namespace UE::Anim
@@ -56,7 +67,7 @@ void FAnimNode_OrientationWarping::GatherDebugData(FNodeDebugData& DebugData)
 {
 	FString DebugLine = DebugData.GetNodeName(this);
 #if ENABLE_ANIM_DEBUG
-	if (CVarAnimNodeOrientationWarpingVerbose.GetValueOnAnyThread() == 1)
+	if (AnimNodeOrientationWarpingVerbose == 1)
 	{
 		if (Mode == EWarpingEvaluationMode::Manual)
 		{
@@ -109,6 +120,14 @@ void FAnimNode_OrientationWarping::UpdateInternal(const FAnimationUpdateContext&
 	UpdateCounter.SynchronizeWith(Context.AnimInstanceProxy->GetUpdateCounter());
 	BlendWeight = Context.GetFinalBlendWeight();
 
+#if ENABLE_ANIM_DEBUG || ENABLE_VISUAL_LOG
+	UE::Anim::FAnimBlendDebugScope* BlendDebugMessage = Context.GetMessage<UE::Anim::FAnimBlendDebugScope>();
+	if (BlendDebugMessage)
+	{
+		DebugColor = BlendDebugMessage->DebugColor;
+	}
+#endif
+
 	if (WarpingSpace == EOrientationWarpingSpace::RootBoneTransform)
 	{
 		if (UE::AnimationWarping::FRootOffsetProvider* RootOffsetProvider = Context.GetMessage<UE::AnimationWarping::FRootOffsetProvider>())
@@ -372,7 +391,7 @@ void FAnimNode_OrientationWarping::EvaluateSkeletalControl_AnyThread(FComponentS
 #else
 	constexpr float DebugDrawScale = 1.f;
 #endif
-	const int32 DebugIndex = CVarAnimNodeOrientationWarpingDebug.GetValueOnAnyThread();
+	const int32 DebugIndex = AnimNodeOrientationWarpingDebug;
 	bDebugging = bDebugging || (DebugIndex > 0);
 
 	if (bDebugging)
@@ -380,6 +399,7 @@ void FAnimNode_OrientationWarping::EvaluateSkeletalControl_AnyThread(FComponentS
 		const FTransform ComponentTransform = Output.AnimInstanceProxy->GetComponentTransform();
 		const FVector ActorForwardDirection = Output.AnimInstanceProxy->GetActorTransform().GetRotation().GetForwardVector();
 		FVector DebugArrowOffset = FVector::ZAxisVector * DebugDrawScale;
+		uint8 DebugAlpha = AnimNodeOrientationWarpingDebugTransparency ? 255 * BlendWeight : 255;
 
 		// Draw debug shapes
 		{
@@ -387,6 +407,9 @@ void FAnimNode_OrientationWarping::EvaluateSkeletalControl_AnyThread(FComponentS
 				? ComponentTransform.GetRotation().RotateVector(LocomotionForward)
 				: ActorForwardDirection;
 
+			Output.AnimInstanceProxy->AnimDrawDebugCircle(
+				ComponentTransform.GetLocation() + DebugArrowOffset + ForwardDirection * 100.f * DebugDrawScale,
+				4.f * DebugDrawScale, 8, DebugColor.WithAlpha(DebugAlpha), FVector::UpVector, false, 0.f, SDPG_World, 1.f);
 			Output.AnimInstanceProxy->AnimDrawDebugDirectionalArrow(
 				ComponentTransform.GetLocation() + DebugArrowOffset,
 				ComponentTransform.GetLocation() + DebugArrowOffset + ForwardDirection * 100.f * DebugDrawScale,
@@ -397,19 +420,26 @@ void FAnimNode_OrientationWarping::EvaluateSkeletalControl_AnyThread(FComponentS
 				: ActorForwardDirection.RotateAngleAxis(OrientationAngle, RotationAxisVector);
 
 			DebugArrowOffset += FVector::ZAxisVector * DebugDrawScale;
+
+			Output.AnimInstanceProxy->AnimDrawDebugCircle(
+				ComponentTransform.GetLocation() + DebugArrowOffset + RotationDirection * 100.f * DebugDrawScale,
+				4.f * DebugDrawScale, 8, DebugColor.WithAlpha(DebugAlpha), FVector::UpVector, false, 0.f, SDPG_World, 1.f);
 			Output.AnimInstanceProxy->AnimDrawDebugDirectionalArrow(
 				ComponentTransform.GetLocation() + DebugArrowOffset,
 				ComponentTransform.GetLocation() + DebugArrowOffset + RotationDirection * 100.f * DebugDrawScale,
-				40.f * DebugDrawScale, FColor::Blue, false, 0.f, 2.f * DebugDrawScale);
+				40.f * DebugDrawScale, FColor::Blue.WithAlpha(DebugAlpha), false, 0.f, 2.f * DebugDrawScale);
 
 			if (bUsedFutureRootMotion && bGraphDrivenWarping)
 			{
 				const FVector FutureRotationDirection = ComponentTransform.GetRotation().RotateVector(FutureRootMotionDeltaDirection);
 
+				Output.AnimInstanceProxy->AnimDrawDebugCircle(
+					ComponentTransform.GetLocation() + DebugArrowOffset + FutureRotationDirection * 100.f * DebugDrawScale,
+					4.f * DebugDrawScale, 8, DebugColor.WithAlpha(DebugAlpha), FVector::UpVector, false, 0.f, SDPG_World, 1.f);
 				Output.AnimInstanceProxy->AnimDrawDebugDirectionalArrow(
 					ComponentTransform.GetLocation() + DebugArrowOffset,
 					ComponentTransform.GetLocation() + DebugArrowOffset + FutureRotationDirection * 100.f * DebugDrawScale,
-					40.f * DebugDrawScale, FColor::Yellow, false, 0.f, 2.f * DebugDrawScale);
+					40.f * DebugDrawScale, FColor::Yellow.WithAlpha(DebugAlpha), false, 0.f, 2.f * DebugDrawScale);
 			}
 
 			const float ActualOrientationAngleDegrees = FMath::RadiansToDegrees(ActualOrientationAngleRad);
@@ -418,10 +448,13 @@ void FAnimNode_OrientationWarping::EvaluateSkeletalControl_AnyThread(FComponentS
 				: ActorForwardDirection.RotateAngleAxis(ActualOrientationAngleDegrees, RotationAxisVector);
 
 			DebugArrowOffset += FVector::ZAxisVector * DebugDrawScale;
+			Output.AnimInstanceProxy->AnimDrawDebugCircle(
+				ComponentTransform.GetLocation() + DebugArrowOffset + WarpedRotationDirection * 100.f * DebugDrawScale,
+				4.f * DebugDrawScale, 8, DebugColor.WithAlpha(DebugAlpha), FVector::UpVector, false, 0.f, SDPG_World, 1.f);
 			Output.AnimInstanceProxy->AnimDrawDebugDirectionalArrow(
 				ComponentTransform.GetLocation() + DebugArrowOffset,
 				ComponentTransform.GetLocation() + DebugArrowOffset + WarpedRotationDirection * 100.f * DebugDrawScale,
-				40.f * DebugDrawScale, FColor::Green, false, 0.f, 2.f * DebugDrawScale);
+				40.f * DebugDrawScale, FColor::Green.WithAlpha(DebugAlpha), false, 0.f, 2.f * DebugDrawScale);
 		}
 
 		// Draw text on mesh in world space
@@ -486,6 +519,7 @@ void FAnimNode_OrientationWarping::EvaluateSkeletalControl_AnyThread(FComponentS
 		const FTransform ComponentTransform = Output.AnimInstanceProxy->GetComponentTransform();
 		const FVector ActorForwardDirection = Output.AnimInstanceProxy->GetActorTransform().GetRotation().GetForwardVector();
 		FVector DebugArrowOffset = FVector::ZAxisVector * DebugDrawScale;
+		uint8 DebugAlpha = AnimNodeOrientationWarpingDebugTransparency ? 255 * BlendWeight : 255;
 
 		// Draw debug shapes
 		{
@@ -493,29 +527,39 @@ void FAnimNode_OrientationWarping::EvaluateSkeletalControl_AnyThread(FComponentS
 				? ComponentTransform.GetRotation().RotateVector(LocomotionForward)
 				: ActorForwardDirection;
 
+			UE_VLOG_CIRCLE_THICK(Output.AnimInstanceProxy->GetAnimInstanceObject(), "OrientationWarping", Display,
+				ComponentTransform.GetLocation() + DebugArrowOffset + ForwardDirection * 100.f * DebugDrawScale,
+				FVector::UpVector, 4.f * DebugDrawScale, DebugColor.WithAlpha(DebugAlpha), 1.0f, TEXT(""));
 			UE_VLOG_ARROW(Output.AnimInstanceProxy->GetAnimInstanceObject(), "OrientationWarping", Display,
 				ComponentTransform.GetLocation() + DebugArrowOffset,
 				ComponentTransform.GetLocation() + DebugArrowOffset + ForwardDirection * 100.f * DebugDrawScale,
-				FColor::Red, TEXT(""));
+				FColor::Red.WithAlpha(DebugAlpha), TEXT(""));
 
 			const FVector RotationDirection = bGraphDrivenWarping
 				? ComponentTransform.GetRotation().RotateVector(RootMotionDeltaDirection)
 				: ActorForwardDirection.RotateAngleAxis(OrientationAngle, RotationAxisVector);
 
 			DebugArrowOffset += FVector::ZAxisVector * DebugDrawScale;
+
+			UE_VLOG_CIRCLE_THICK(Output.AnimInstanceProxy->GetAnimInstanceObject(), "OrientationWarping", Display,
+				ComponentTransform.GetLocation() + DebugArrowOffset + RotationDirection * 100.f * DebugDrawScale,
+				FVector::UpVector, 4.f * DebugDrawScale, DebugColor.WithAlpha(DebugAlpha), 1.0f, TEXT(""));
 			UE_VLOG_ARROW(Output.AnimInstanceProxy->GetAnimInstanceObject(), "OrientationWarping", Display,
 				ComponentTransform.GetLocation() + DebugArrowOffset,
 				ComponentTransform.GetLocation() + DebugArrowOffset + RotationDirection * 100.f * DebugDrawScale,
-				FColor::Blue, TEXT(""));
+				FColor::Blue.WithAlpha(DebugAlpha), TEXT(""));
 
 			if (bUsedFutureRootMotion && bGraphDrivenWarping)
 			{
 				const FVector FutureRotationDirection = ComponentTransform.GetRotation().RotateVector(FutureRootMotionDeltaDirection);
 
-				Output.AnimInstanceProxy->AnimDrawDebugDirectionalArrow(
+				UE_VLOG_CIRCLE_THICK(Output.AnimInstanceProxy->GetAnimInstanceObject(), "OrientationWarping", Display,
+					ComponentTransform.GetLocation() + DebugArrowOffset + FutureRotationDirection * 100.f * DebugDrawScale,
+					FVector::UpVector, 4.f * DebugDrawScale, DebugColor.WithAlpha(DebugAlpha), 1.0f, TEXT(""));
+				UE_VLOG_ARROW(Output.AnimInstanceProxy->GetAnimInstanceObject(), "OrientationWarping", Display,
 					ComponentTransform.GetLocation() + DebugArrowOffset,
 					ComponentTransform.GetLocation() + DebugArrowOffset + FutureRotationDirection * 100.f * DebugDrawScale,
-					40.f * DebugDrawScale, FColor::Yellow, false, 0.f, 2.f * DebugDrawScale);
+					FColor::Yellow.WithAlpha(DebugAlpha), TEXT(""));
 			}
 
 			const float ActualOrientationAngleDegrees = FMath::RadiansToDegrees(ActualOrientationAngleRad);
@@ -525,10 +569,13 @@ void FAnimNode_OrientationWarping::EvaluateSkeletalControl_AnyThread(FComponentS
 
 			DebugArrowOffset += FVector::ZAxisVector * DebugDrawScale;
 
+			UE_VLOG_CIRCLE_THICK(Output.AnimInstanceProxy->GetAnimInstanceObject(), "OrientationWarping", Display,
+				ComponentTransform.GetLocation() + DebugArrowOffset + WarpedRotationDirection * 100.f * DebugDrawScale,
+				FVector::UpVector, 4.f * DebugDrawScale, DebugColor.WithAlpha(DebugAlpha), 1.0f, TEXT(""));
 			UE_VLOG_ARROW(Output.AnimInstanceProxy->GetAnimInstanceObject(), "OrientationWarping", Display,
 				ComponentTransform.GetLocation() + DebugArrowOffset,
 				ComponentTransform.GetLocation() + DebugArrowOffset + WarpedRotationDirection * 100.f * DebugDrawScale,
-				FColor::Green, TEXT(""));
+				FColor::Green.WithAlpha(DebugAlpha), TEXT(""));
 		}
 	}
 #endif
@@ -609,7 +656,7 @@ void FAnimNode_OrientationWarping::EvaluateSkeletalControl_AnyThread(FComponentS
 bool FAnimNode_OrientationWarping::IsValidToEvaluate(const USkeleton* Skeleton, const FBoneContainer& RequiredBones)
 {
 #if ENABLE_ANIM_DEBUG
-	if (CVarAnimNodeOrientationWarpingEnable.GetValueOnAnyThread() == 0)
+	if (AnimNodeOrientationWarpingEnable == 0)
 	{
 		return false;
 	}
diff --git a/Engine/Plugins/Animation/AnimationWarping/Source/Runtime/Public/BoneControllers/AnimNode_OrientationWarping.h b/Engine/Plugins/Animation/AnimationWarping/Source/Runtime/Public/BoneControllers/AnimNode_OrientationWarping.h
--- a/Engine/Plugins/Animation/AnimationWarping/Source/Runtime/Public/BoneControllers/AnimNode_OrientationWarping.h
+++ b/Engine/Plugins/Animation/AnimationWarping/Source/Runtime/Public/BoneControllers/AnimNode_OrientationWarping.h
@@ -228,9 +228,12 @@ private:
 	float CounterCompensateTargetAngleRad = 0.0f;
 
 #if ENABLE_ANIM_DEBUG || ENABLE_VISUAL_LOG
-	// Store these in debug only so that they can persist in debug / across pauses
+	// Store prediction in debug only so that they can persist across pauses
 	FVector FutureRootMotionDeltaDirection = FVector::ZeroVector;
 	bool bUsedFutureRootMotion = false;
+
+	// Stored from blend message, used to match with rewind debugger
+	FColor DebugColor = FColor::Black;
 #endif
 
 	// Internal orientation warping angle
diff --git a/Engine/Plugins/Animation/BlendStack/Source/Runtime/Private/AnimNode_BlendStack.cpp b/Engine/Plugins/Animation/BlendStack/Source/Runtime/Private/AnimNode_BlendStack.cpp
--- a/Engine/Plugins/Animation/BlendStack/Source/Runtime/Private/AnimNode_BlendStack.cpp
+++ b/Engine/Plugins/Animation/BlendStack/Source/Runtime/Private/AnimNode_BlendStack.cpp
@@ -3,6 +3,7 @@
 #include "BlendStack/AnimNode_BlendStack.h"
 #include "Algo/MaxElement.h"
 #include "Animation/AnimComposite.h"
+#include "Animation/AnimBlendDebugScope.h"
 #include "Animation/AnimInstanceProxy.h"
 #include "Animation/AnimMontage.h"
 #include "Animation/AnimNode_Inertialization.h"
@@ -117,6 +118,20 @@ void FBlendStackAnimPlayer::Initialize(const FAnimationInitializeContext& Contex
 	PoseLinkIndex = InPoseLinkIdx;
 
 	OverrideCurve.Empty();
+
+#if ENABLE_ANIM_DEBUG || ENABLE_VISUAL_LOG 
+	// Matches 'MakeBlendWeightCurveColor' in DebugWeightsTrack.cpp to match rewind debugger color (Does not depend since do not want dependency on debug for color).
+	auto MakeDebugColor = [](uint32 InSeed)
+	{
+		FRandomStream Stream(InSeed);
+		const uint8 Hue = (uint8)(Stream.FRand() * 255.0f);
+		const uint8 SatVal = 196;
+		return FLinearColor::MakeFromHSV8(Hue, SatVal, SatVal);
+	};
+
+	uint64 AssetId = FObjectTrace::GetObjectId(AnimationAsset);
+	DebugColor = MakeDebugColor(CityHash32(reinterpret_cast<char*>(&AssetId), 8)).ToFColor(true);
+#endif
 }
 
 void FBlendStackAnimPlayer::UpdatePlayRate(float PlayRate)
@@ -353,6 +368,13 @@ FAnimNode_AssetPlayerBase* FBlendStackAnimPlayer::GetAssetPlayerNode()
 	return nullptr;
 }
 
+#if ENABLE_ANIM_DEBUG || ENABLE_VISUAL_LOG 
+FColor FBlendStackAnimPlayer::GetDebugColor() const
+{
+	return DebugColor;
+}
+#endif
+
 void FBlendStackAnimPlayer::UpdateWithDeltaTime(float DeltaTime, int32 PlayerDepth, float PlayerDepthBlendInTimeMultiplier)
 {
 	const bool bIsMainPlayer = PlayerDepth == 0;
@@ -738,7 +760,13 @@ void FAnimNode_BlendStack_Standalone::UpdateAssetPlayer(const FAnimationUpdateCo
 		const float AnimPlayerBlendWeight = CurrentWeightMultiplier * BlendInPercentage;
 
 		FAnimationUpdateContext AnimPlayerContext = Context.FractionalWeightAndRootMotion(AnimPlayerBlendWeight, AnimPlayerBlendWeight);
-		UpdateSample((AnimPlayerIndex == 0) ? AnimPlayerContext : AnimPlayerContext.AsInactive(), AnimPlayerIndex);
+		{
+
+#if ENABLE_ANIM_DEBUG || ENABLE_VISUAL_LOG 
+			UE::Anim::TScopedGraphMessage<UE::Anim::FAnimBlendDebugScope> BlendDebugMessage(Context, Context, AnimPlayerIndex, BlendStackSize, AnimPlayer.GetDebugColor());
+#endif
+			UpdateSample((AnimPlayerIndex == 0) ? AnimPlayerContext : AnimPlayerContext.AsInactive(), AnimPlayerIndex);
+		}
 		CurrentWeightMultiplier *= (1.f - BlendInPercentage);
 
 		++AnimPlayerIndex;
diff --git a/Engine/Plugins/Animation/BlendStack/Source/Runtime/Public/BlendStack/AnimNode_BlendStack.h b/Engine/Plugins/Animation/BlendStack/Source/Runtime/Public/BlendStack/AnimNode_BlendStack.h
--- a/Engine/Plugins/Animation/BlendStack/Source/Runtime/Public/BlendStack/AnimNode_BlendStack.h
+++ b/Engine/Plugins/Animation/BlendStack/Source/Runtime/Public/BlendStack/AnimNode_BlendStack.h
@@ -68,6 +68,10 @@ struct BLENDSTACK_API FBlendStackAnimPlayer
 
 	FAnimNode_AssetPlayerBase* GetAssetPlayerNode();
 
+#if ENABLE_ANIM_DEBUG || ENABLE_VISUAL_LOG 
+	FColor GetDebugColor() const;
+#endif
+
 	// Curves to add to the pose after the player evaluates
 	TBaseBlendedCurve<FDefaultAllocator, UE::Anim::FCurveElement> OverrideCurve;
 
@@ -101,6 +105,10 @@ private:
 	float TotalBlendInTime = 0.f;
 	float CurrentBlendInTime = 0.f;
 	float TimeToActivation = 0.f;
+
+#if ENABLE_ANIM_DEBUG || ENABLE_VISUAL_LOG 
+	FColor DebugColor = FColor::Black;
+#endif
 };
 
 namespace UE::BlendStack
diff --git a/Engine/Source/Runtime/Engine/Private/Animation/AnimBlendDebugScope.cpp b/Engine/Source/Runtime/Engine/Private/Animation/AnimBlendDebugScope.cpp
new file mode 100644
--- /dev/null
+++ b/Engine/Source/Runtime/Engine/Private/Animation/AnimBlendDebugScope.cpp
@@ -0,0 +1,20 @@
+// Copyright Epic Games, Inc. All Rights Reserved.
+
+#include "Animation/AnimBlendDebugScope.h"
+#include "Animation/AnimNodeBase.h"
+
+IMPLEMENT_ANIMGRAPH_MESSAGE(UE::Anim::FAnimBlendDebugScope);
+
+namespace UE { namespace Anim {
+
+
+FAnimBlendDebugScope::FAnimBlendDebugScope(const FAnimationBaseContext& InContext, const int32& InBlendIndex, const int32& InNumBlends, const FColor& InDebugColor)
+	: BlendIndex(InBlendIndex)
+	, NumBlends(InNumBlends)
+	, DebugColor(InDebugColor)
+	, NodeId(InContext.GetCurrentNodeId())
+{
+}
+
+
+}}	// namespace UE::Anim
diff --git a/Engine/Source/Runtime/Engine/Private/Components/LineBatchComponent.cpp b/Engine/Source/Runtime/Engine/Private/Components/LineBatchComponent.cpp
--- a/Engine/Source/Runtime/Engine/Private/Components/LineBatchComponent.cpp
+++ b/Engine/Source/Runtime/Engine/Private/Components/LineBatchComponent.cpp
@@ -73,7 +73,7 @@ void FLineBatcherSceneProxy::GetDynamicMeshElements(const TArray<const FSceneVie
 #endif
 			for (const FBatchedLine& Line: Lines)
 			{
-				PDI->DrawLine(Line.Start,Line.End, Line.Color, Line.DepthPriority, Line.Thickness);
+				PDI->DrawTranslucentLine(Line.Start,Line.End, Line.Color, Line.DepthPriority, Line.Thickness);
 			}
 
 			for (const FBatchedPoint& Point: Points)
diff --git a/Engine/Source/Runtime/Engine/Private/Components/SplineComponent.cpp b/Engine/Source/Runtime/Engine/Private/Components/SplineComponent.cpp
--- a/Engine/Source/Runtime/Engine/Private/Components/SplineComponent.cpp
+++ b/Engine/Source/Runtime/Engine/Private/Components/SplineComponent.cpp
@@ -2132,9 +2132,9 @@ void USplineComponent::Draw(FPrimitiveDrawInterface* PDI, const FSceneView* View
 					const float Key = (KeyIdx - 1) + (StepIdx / static_cast<float>(NumSteps));
 					const FVector NewPos = LocalToWorld.TransformPosition(SplineInfo.Eval(Key, FVector(0)));
 #if WITH_EDITOR
-					PDI->DrawLine(OldPos, NewPos, LineColor, DepthPriorityGroup, SegmentLineThickness);
+					PDI->DrawTranslucentLine(OldPos, NewPos, LineColor, DepthPriorityGroup, SegmentLineThickness);
 #else
-					PDI->DrawLine(OldPos, NewPos, LineColor, DepthPriorityGroup);
+					PDI->DrawTranslucentLine(OldPos, NewPos, LineColor, DepthPriorityGroup);
 #endif
 					OldPos = NewPos;
 				}
diff --git a/Engine/Source/Runtime/Engine/Public/Animation/AnimBlendDebugScope.h b/Engine/Source/Runtime/Engine/Public/Animation/AnimBlendDebugScope.h
new file mode 100644
--- /dev/null
+++ b/Engine/Source/Runtime/Engine/Public/Animation/AnimBlendDebugScope.h
@@ -0,0 +1,35 @@
+// Copyright Epic Games, Inc. All Rights Reserved.
+
+#pragma once
+
+#include "CoreMinimal.h"
+#include "AnimNodeMessages.h"
+
+struct FAnimationBaseContext;
+
+namespace UE { namespace Anim {
+
+// Scoped graph message used to synchronize animations at various points in an anim graph
+class FAnimBlendDebugScope : public IGraphMessage
+{
+	DECLARE_ANIMGRAPH_MESSAGE_API(FAnimBlendDebugScope, ENGINE_API);
+
+public:
+	ENGINE_API FAnimBlendDebugScope(const FAnimationBaseContext& InContext, const int32& InBlendIndex = INDEX_NONE, const int32& InNumBlends = 0, const FColor& InDebugColor = FColor::Black);
+
+	// The index for this anim, if any
+	int32 BlendIndex = INDEX_NONE;
+
+	// Number of current blends, if any
+	int32 NumBlends = 0;
+
+	// A color used to identify this anim
+	FColor DebugColor = FColor::Black;
+
+private:
+	// The node ID that was used when this scope was entered
+	int32 NodeId;
+
+};
+
+}}	// namespace UE::Anim
