Add Ground Normal and On Ground to OffsetRootBone node to control previously hard coded ground projection feature

diff --git a/Engine/Plugins/Animation/AnimationWarping/Source/Runtime/Private/BoneControllers/AnimNode_OffsetRootBone.cpp b/Engine/Plugins/Animation/AnimationWarping/Source/Runtime/Private/BoneControllers/AnimNode_OffsetRootBone.cpp
--- a/Engine/Plugins/Animation/AnimationWarping/Source/Runtime/Private/BoneControllers/AnimNode_OffsetRootBone.cpp
+++ b/Engine/Plugins/Animation/AnimationWarping/Source/Runtime/Private/BoneControllers/AnimNode_OffsetRootBone.cpp
@@ -169,7 +169,7 @@ void FAnimNode_OffsetRootBone::Evaluate_AnyThread(FPoseContext& Output)
 		RootMotionTransformDelta = FTransform(GetRotationDelta(), GetTranslationDelta());
 	}
 	float MaxTranslationOffset = GetMaxTranslationError();
-
+	
 	bool bCollisionDetected = false;
 	FVector CollisionPoint;
 	FVector CollisionNormal;
@@ -279,10 +279,10 @@ void FAnimNode_OffsetRootBone::Evaluate_AnyThread(FPoseContext& Output)
 	SimulatedTranslation = SimulatedTransform.GetLocation();
 	SimulatedRotation = SimulatedTransform.GetRotation();
 	
-	// TODO: Make this a parameter
-	const FVector GravityDirCS = -FVector::UpVector;
-	// Simulated translation should stay the same along the approach direction
-	SimulatedTranslation = FVector::PointPlaneProject(SimulatedTranslation, ComponentTransform.GetLocation(), GravityDirCS);
+	if (GetOnGround())
+	{
+		SimulatedTranslation = FVector::PointPlaneProject(SimulatedTranslation, ComponentTransform.GetLocation(), GetGroundNormal());
+	}
 
 	const FBoneContainer& RequiredBones = Output.Pose.GetBoneContainer();
 
@@ -510,6 +510,16 @@ bool FAnimNode_OffsetRootBone::GetResetEveryFrame() const
 	return GET_ANIM_NODE_DATA(bool, bResetEveryFrame);
 }
 
+bool FAnimNode_OffsetRootBone::GetOnGround() const
+{
+	return GET_ANIM_NODE_DATA(bool, bOnGround);
+}
+
+const FVector& FAnimNode_OffsetRootBone::GetGroundNormal() const
+{
+	return GET_ANIM_NODE_DATA(FVector, GroundNormal);
+}
+
 const FVector& FAnimNode_OffsetRootBone::GetTranslationDelta() const
 {
 	return GET_ANIM_NODE_DATA(FVector, TranslationDelta);
diff --git a/Engine/Plugins/Animation/AnimationWarping/Source/Runtime/Public/BoneControllers/AnimNode_OffsetRootBone.h b/Engine/Plugins/Animation/AnimationWarping/Source/Runtime/Public/BoneControllers/AnimNode_OffsetRootBone.h
--- a/Engine/Plugins/Animation/AnimationWarping/Source/Runtime/Public/BoneControllers/AnimNode_OffsetRootBone.h
+++ b/Engine/Plugins/Animation/AnimationWarping/Source/Runtime/Public/BoneControllers/AnimNode_OffsetRootBone.h
@@ -129,6 +129,14 @@ struct ANIMATIONWARPINGRUNTIME_API FAnimNode_OffsetRootBone : public FAnimNode_B
 	UPROPERTY(EditAnywhere, Category = Settings, meta = (EditCondition = bClampToRotationVelocity, FoldProperty, PinHiddenByDefault))
 	float RotationSpeedRatio = 0.5f;
 
+	// When OnGround is true, root motion velicities will be projected onto the ground surface
+	UPROPERTY(EditAnywhere, Category = Settings, meta = (FoldProperty, PinHiddenByDefault))
+	bool bOnGround = true;
+	
+	// Surface normal of the ground
+	UPROPERTY(EditAnywhere, Category = Settings, meta = (FoldProperty, PinHiddenByDefault))
+	FVector GroundNormal = { 0, 0, 1 };
+
 	// Delta applied to the translation offset this frame. 
 	// For procedural values, consider adjusting the input by delta time.
 	UPROPERTY(EditAnywhere, Category = Evaluation, meta = (FoldProperty, PinHiddenByDefault))
@@ -161,6 +169,8 @@ public:
 	// Folded property accesors
 	EWarpingEvaluationMode GetEvaluationMode() const;
 	bool GetResetEveryFrame() const;
+	bool GetOnGround() const;
+	const FVector& GetGroundNormal() const;
 	const FVector& GetTranslationDelta() const;
 	const FRotator& GetRotationDelta() const;
 	EOffsetRootBoneMode GetTranslationMode() const;
