Evaluation Notifies, and root motion based Alignment implementation in anim graph

diff --git a/Engine/Plugins/Animation/AnimationWarping/Source/Editor/Private/AnimGraph/AnimGraphNode_EvaluationNotifies.cpp b/Engine/Plugins/Animation/AnimationWarping/Source/Editor/Private/AnimGraph/AnimGraphNode_EvaluationNotifies.cpp
new file mode 100644
--- /dev/null
+++ b/Engine/Plugins/Animation/AnimationWarping/Source/Editor/Private/AnimGraph/AnimGraphNode_EvaluationNotifies.cpp
@@ -0,0 +1,70 @@
+// Copyright Epic Games, Inc. All Rights Reserved.
+
+#include "AnimGraph/AnimGraphNode_EvaluationNotifies.h"
+#include "Animation/AnimRootMotionProvider.h"
+#include "Kismet2/BlueprintEditorUtils.h"
+#include "Kismet2/CompilerResultsLog.h"
+#include "DetailCategoryBuilder.h"
+#include "DetailLayoutBuilder.h"
+#include "PropertyHandle.h"
+#include "ScopedTransaction.h"
+
+#define LOCTEXT_NAMESPACE "AnimationWarping"
+
+UAnimGraphNode_EvaluationNotifies::UAnimGraphNode_EvaluationNotifies(const FObjectInitializer& ObjectInitializer) : Super(ObjectInitializer)
+{
+}
+
+FText UAnimGraphNode_EvaluationNotifies::GetControllerDescription() const
+{
+	return LOCTEXT("EvaluationNotifies", "EvaluationNotifies");
+}
+
+FText UAnimGraphNode_EvaluationNotifies::GetTooltipText() const
+{
+	return LOCTEXT("EvaluationNotifiesTooltip", "Rotates the root and lower body by the specified angle, while counter rotating the upper body to maintain the forward facing direction.");
+}
+
+FText UAnimGraphNode_EvaluationNotifies::GetNodeTitle(ENodeTitleType::Type TitleType) const
+{
+	return GetControllerDescription();
+}
+
+FLinearColor UAnimGraphNode_EvaluationNotifies::GetNodeTitleColor() const
+{
+	return FLinearColor(FColor(153.f, 0.f, 0.f));
+}
+
+void UAnimGraphNode_EvaluationNotifies::CustomizePinData(UEdGraphPin* Pin, FName SourcePropertyName, int32 ArrayIndex) const
+{
+	Super::CustomizePinData(Pin, SourcePropertyName, ArrayIndex);
+}
+
+void UAnimGraphNode_EvaluationNotifies::CustomizeDetails(IDetailLayoutBuilder& DetailBuilder)
+{
+	DECLARE_SCOPE_HIERARCHICAL_COUNTER_FUNC()
+	Super::CustomizeDetails(DetailBuilder);
+}
+
+void UAnimGraphNode_EvaluationNotifies::PostEditChangeProperty(struct FPropertyChangedEvent& PropertyChangedEvent)
+{
+	DECLARE_SCOPE_HIERARCHICAL_COUNTER_FUNC()
+	Super::PostEditChangeProperty(PropertyChangedEvent);
+}
+
+void UAnimGraphNode_EvaluationNotifies::GetInputLinkAttributes(FNodeAttributeArray& OutAttributes) const
+{
+	OutAttributes.Add(UE::Anim::IAnimRootMotionProvider::AttributeName);
+}
+
+void UAnimGraphNode_EvaluationNotifies::GetOutputLinkAttributes(FNodeAttributeArray& OutAttributes) const
+{
+	OutAttributes.Add(UE::Anim::IAnimRootMotionProvider::AttributeName);
+}
+
+void UAnimGraphNode_EvaluationNotifies::ValidateAnimNodeDuringCompilation(USkeleton* ForSkeleton, FCompilerResultsLog& MessageLog)
+{
+	Super::ValidateAnimNodeDuringCompilation(ForSkeleton, MessageLog);
+}
+
+#undef LOCTEXT_NAMESPACE
diff --git a/Engine/Plugins/Animation/AnimationWarping/Source/Editor/Public/AnimGraph/AnimGraphNode_EvaluationNotifies.h b/Engine/Plugins/Animation/AnimationWarping/Source/Editor/Public/AnimGraph/AnimGraphNode_EvaluationNotifies.h
new file mode 100644
--- /dev/null
+++ b/Engine/Plugins/Animation/AnimationWarping/Source/Editor/Public/AnimGraph/AnimGraphNode_EvaluationNotifies.h
@@ -0,0 +1,39 @@
+// Copyright Epic Games, Inc. All Rights Reserved.
+
+#pragma once
+#include "AnimGraphNode_SkeletalControlBase.h"
+#include "EvaluationNotifies/AnimNode_EvaluationNotifies.h"
+#include "AnimGraphNode_EvaluationNotifies.generated.h"
+
+namespace ENodeTitleType { enum Type : int; }
+
+UCLASS()
+class ANIMATIONWARPINGEDITOR_API UAnimGraphNode_EvaluationNotifies : public UAnimGraphNode_SkeletalControlBase
+{
+	GENERATED_UCLASS_BODY()
+
+	UPROPERTY(EditAnywhere, Category = Settings)
+	FAnimNode_EvaluationNotifies Node;
+
+public:
+	// UEdGraphNode interface
+	virtual FText GetNodeTitle(ENodeTitleType::Type TitleType) const override;
+	virtual FLinearColor GetNodeTitleColor() const override;
+	virtual FText GetTooltipText() const override;
+	virtual void CustomizeDetails(IDetailLayoutBuilder& DetailBuilder) override;
+	virtual void PostEditChangeProperty(struct FPropertyChangedEvent& PropertyChangedEvent) override;
+	virtual void GetInputLinkAttributes(FNodeAttributeArray& OutAttributes) const override;
+	virtual void GetOutputLinkAttributes(FNodeAttributeArray& OutAttributes) const override;
+	virtual void ValidateAnimNodeDuringCompilation(USkeleton* ForSkeleton, FCompilerResultsLog& MessageLog) override;
+	// End of UEdGraphNode interface
+
+protected:
+	// UAnimGraphNode_Base interface
+	virtual void CustomizePinData(UEdGraphPin* Pin, FName SourcePropertyName, int32 ArrayIndex) const override;
+	// End of UAnimGraphNode_Base interface
+
+	// UAnimGraphNode_SkeletalControlBase interface
+	virtual FText GetControllerDescription() const override;
+	virtual const FAnimNode_SkeletalControlBase* GetNode() const override { return &Node; }
+	// End of UAnimGraphNode_SkeletalControlBase interface
+};
diff --git a/Engine/Plugins/Animation/AnimationWarping/Source/Runtime/Private/AnimationWarpingRuntimeModule.cpp b/Engine/Plugins/Animation/AnimationWarping/Source/Runtime/Private/AnimationWarpingRuntimeModule.cpp
--- a/Engine/Plugins/Animation/AnimationWarping/Source/Runtime/Private/AnimationWarpingRuntimeModule.cpp
+++ b/Engine/Plugins/Animation/AnimationWarping/Source/Runtime/Private/AnimationWarpingRuntimeModule.cpp
@@ -7,6 +7,8 @@
 #include "Animation/BuiltInAttributeTypes.h"
 #include "Features/IModularFeatures.h"
 #include "Modules/ModuleManager.h"
+#include "EvaluationNotifies/AnimNode_EvaluationNotifies.h"
+#include "EvaluationNotifies/AnimNotifyState_Alignment.h"
 
 class FAnimationWarpingRuntimeModule : public IAnimationWarpingRuntimeModule
 {
@@ -42,11 +44,15 @@ FModule::FModule()
 void FModule::StartupModule()
 {
 	IModularFeatures::Get().RegisterModularFeature(UE::Anim::IAnimRootMotionProvider::ModularFeatureName, this);
+
+	FAnimNode_EvaluationNotifies::RegisterEvaluationHandler(UNotifyState_Alignment::StaticClass(), FAlignmentNotifyInstance::StaticStruct());
 }
 
 void FModule::ShutdownModule()
 {
 	IModularFeatures::Get().UnregisterModularFeature(UE::Anim::IAnimRootMotionProvider::ModularFeatureName, this);
+	
+	FAnimNode_EvaluationNotifies::UnregisterEvaluationHandler(UNotifyState_Alignment::StaticClass());
 }
 
 void FModule::SampleRootMotion(const FDeltaTimeRecord& SampleRange, const UAnimSequence& Sequence, bool bLoopingSequence, UE::Anim::FStackAttributeContainer& OutAttributes) const
diff --git a/Engine/Plugins/Animation/AnimationWarping/Source/Runtime/Private/EvaluationNotifies/AnimNode_EvaluationNotifies.cpp b/Engine/Plugins/Animation/AnimationWarping/Source/Runtime/Private/EvaluationNotifies/AnimNode_EvaluationNotifies.cpp
new file mode 100644
--- /dev/null
+++ b/Engine/Plugins/Animation/AnimationWarping/Source/Runtime/Private/EvaluationNotifies/AnimNode_EvaluationNotifies.cpp
@@ -0,0 +1,122 @@
+// Copyright Epic Games, Inc. All Rights Reserved.
+
+#include "EvaluationNotifies/AnimNode_EvaluationNotifies.h"
+#include "Animation/AnimInstanceProxy.h"
+#include "Animation/AnimRootMotionProvider.h"
+#include "Animation/AnimTrace.h"
+#include "BoneControllers/AnimNode_OffsetRootBone.h"
+#include "Curves/BezierUtilities.h"
+#include "StructUtils/InstancedStruct.h"
+
+#include UE_INLINE_GENERATED_CPP_BY_NAME(AnimNode_EvaluationNotifies)
+
+TMap<UClass*, UScriptStruct*> FAnimNode_EvaluationNotifies::NotifyEvaluationHandlerMap;
+
+void FAnimNode_EvaluationNotifies::GetEvaluationNotifiesFromAnimation(const UAnimSequenceBase* Animation, TArray<FInstancedStruct>& OutNotifyInstances)
+{
+	if(Animation)
+	{
+		// todo: we should probably preprocess sequences for this number
+		int32 Count = 0;
+		for (int32 Idx = 0; Idx < Animation->Notifies.Num(); Idx++)
+		{
+			const FAnimNotifyEvent& NotifyEvent = Animation->Notifies[Idx];
+			if (UAnimNotifyState* Notify = Cast<UAnimNotifyState>(NotifyEvent.NotifyStateClass))
+			{
+				if (UScriptStruct** EvaluationHandlerStruct = NotifyEvaluationHandlerMap.Find(Notify->GetClass()))
+				{
+					Count++;
+				}
+			}
+		}
+
+		OutNotifyInstances.Empty(Count);
+
+		for (int32 Idx = 0; Idx < Animation->Notifies.Num(); Idx++)
+		{
+			const FAnimNotifyEvent& NotifyEvent = Animation->Notifies[Idx];
+			if (UAnimNotifyState* Notify = Cast<UAnimNotifyState>(NotifyEvent.NotifyStateClass))
+			{
+				if (UScriptStruct** EvaluationHandlerStruct = NotifyEvaluationHandlerMap.Find(Notify->GetClass()))
+				{
+					FInstancedStruct& Struct = OutNotifyInstances.AddDefaulted_GetRef();
+					Struct.InitializeAs(*EvaluationHandlerStruct);
+					FEvaluationNotifyInstance& Instance = Struct.GetMutable<FEvaluationNotifyInstance>();
+					Instance.AnimNotify = Notify;
+					Instance.StartTime = NotifyEvent.GetTriggerTime();
+					Instance.EndTime = NotifyEvent.GetEndTriggerTime();
+				}
+			}
+		}
+	}
+}
+
+void FAnimNode_EvaluationNotifies::UpdateInternal(const FAnimationUpdateContext& Context)
+{
+	FAnimNode_SkeletalControlBase::UpdateInternal(Context);
+
+	if (UE::AnimationWarping::FRootOffsetProvider* RootOffsetProvider = Context.GetMessage<UE::AnimationWarping::FRootOffsetProvider>())
+	{
+		RootBoneTransform = RootOffsetProvider->GetRootTransform();
+	}
+	else
+	{
+		RootBoneTransform = Context.AnimInstanceProxy->GetComponentTransform();
+	}
+}
+
+void FAnimNode_EvaluationNotifies::Initialize_AnyThread(const FAnimationInitializeContext& Context)
+{
+	FAnimNode_SkeletalControlBase::Initialize_AnyThread(Context);
+}
+
+
+void FAnimNode_EvaluationNotifies::EvaluateSkeletalControl_AnyThread(FComponentSpacePoseContext& Output, TArray<FBoneTransform>& OutBoneTransforms)
+{
+	const float DeltaSeconds = Output.AnimInstanceProxy->GetDeltaSeconds();
+
+	if (DeltaSeconds > 0.f)
+	{
+		const UAnimSequenceBase* AnimSequenceBase = Cast<UAnimSequenceBase>(CurrentAnimAsset);
+		if (AnimSequenceBase != CurrentSequence)
+		{
+			PreviousAnimAssetTime = CurrentAnimAssetTime;
+			CurrentSequence = AnimSequenceBase;
+			GetEvaluationNotifiesFromAnimation(CurrentSequence, Tags);
+		}
+
+		for(auto& TagData : Tags)
+		{
+			FEvaluationNotifyInstance& Tag = TagData.GetMutable<FEvaluationNotifyInstance>();
+			if (Tag.StartTime <= CurrentAnimAssetTime && Tag.EndTime >= CurrentAnimAssetTime)
+			{
+				if (!Tag.bActive)
+				{
+					Tag.bActive = true;
+					Tag.Start(CurrentSequence);
+				}
+
+				Tag.Update(CurrentSequence, CurrentAnimAssetTime, DeltaSeconds, RootBoneTransform, NamedTransforms, Output, OutBoneTransforms);
+			}
+			else
+			{
+				if (Tag.bActive)
+				{
+					Tag.bActive = false;
+					Tag.End();
+				}
+			}
+			
+		}
+	}
+}
+
+void FAnimNode_EvaluationNotifies::RegisterEvaluationHandler(UClass* NotifyType, UScriptStruct* HandlerType)
+{
+	NotifyEvaluationHandlerMap.Add(NotifyType, HandlerType);
+}
+
+void FAnimNode_EvaluationNotifies::UnregisterEvaluationHandler(UClass* NotifyType)
+{
+	NotifyEvaluationHandlerMap.Remove(NotifyType);
+}
diff --git a/Engine/Plugins/Animation/AnimationWarping/Source/Runtime/Private/EvaluationNotifies/AnimNotifyState_Alignment.cpp b/Engine/Plugins/Animation/AnimationWarping/Source/Runtime/Private/EvaluationNotifies/AnimNotifyState_Alignment.cpp
new file mode 100644
--- /dev/null
+++ b/Engine/Plugins/Animation/AnimationWarping/Source/Runtime/Private/EvaluationNotifies/AnimNotifyState_Alignment.cpp
@@ -0,0 +1,59 @@
+// Copyright Epic Games, Inc. All Rights Reserved.
+
+#include "EvaluationNotifies/AnimNotifyState_Alignment.h"
+#include "Animation/AnimInstanceProxy.h"
+#include "Animation/AnimRootMotionProvider.h"
+
+#include UE_INLINE_GENERATED_CPP_BY_NAME(AnimNotifyState_Alignment)
+
+void FAlignmentNotifyInstance::Start(const UAnimSequenceBase* AnimationAsset)
+{
+	UNotifyState_Alignment* AlignmentNotify = Cast<UNotifyState_Alignment>(AnimNotify);
+	AlignBone = AlignmentNotify->AlignBone;
+}
+
+void FAlignmentNotifyInstance::Update(const UAnimSequenceBase* AnimationAsset, float CurrentTime, float DeltaTime, 
+                                   FTransform& RootBoneTransform, const TMap<FName, FTransform>& NamedTransforms, FComponentSpacePoseContext& Output, TArray<FBoneTransform>& OutBoneTransforms)
+{
+	UNotifyState_Alignment* AlignmentNotify = Cast<UNotifyState_Alignment>(AnimNotify);
+	if (!AlignBone.HasValidSetup())
+	{
+		const FBoneContainer& RequiredBones = Output.AnimInstanceProxy->GetRequiredBones();	
+		AlignBone.Initialize(RequiredBones);
+	}
+
+	const UE::Anim::IAnimRootMotionProvider* RootMotionProvider = UE::Anim::IAnimRootMotionProvider::Get();
+    ensureMsgf(RootMotionProvider, TEXT("Alignment expected a valid root motion delta provider interface."));
+
+    if (RootMotionProvider)
+    {
+    	FTransform ThisFrameRootMotionTransform;
+    	if (RootMotionProvider->ExtractRootMotion(Output.CustomAttributes, ThisFrameRootMotionTransform))
+    	{
+    		if (const FTransform* TargetTransform = NamedTransforms.Find(AlignmentNotify->TransformName))
+    		{
+    			// simple linear blend weight
+    			float TimeRemaining = EndTime-CurrentTime;
+    			float Weight = TimeRemaining <= 0.f ? 1.f : FMath::Min(1.f, DeltaTime / TimeRemaining);
+    		
+				const FBoneContainer& RequiredBones = Output.AnimInstanceProxy->GetRequiredBones();	
+    			FCompactPoseBoneIndex AlignBoneIndex = AlignBone.GetCompactPoseIndex(RequiredBones);
+    			if (AlignBoneIndex.IsValid())
+    			{
+    				const FTransform& AlignBoneTransform  = Output.Pose.GetComponentSpaceTransform(AlignBoneIndex);
+    				RootBoneTransform = AlignBoneTransform * ThisFrameRootMotionTransform * RootBoneTransform;;
+    			}
+
+    			RootBoneTransform = AlignmentNotify->AlignOffset * RootBoneTransform;
+    		
+    			FTransform RelativeOffset = TargetTransform->GetRelativeTransform(RootBoneTransform);
+
+    			FTransform ScaledOffset = FTransform(FMath::Lerp(FQuat::Identity, RelativeOffset.GetRotation(), Weight),
+														Weight * RelativeOffset.GetTranslation());
+
+    			RootMotionProvider->OverrideRootMotion(ThisFrameRootMotionTransform * ScaledOffset, Output.CustomAttributes);
+    		}
+    	}
+    }
+}
+
diff --git a/Engine/Plugins/Animation/AnimationWarping/Source/Runtime/Public/EvaluationNotifies/AnimNode_EvaluationNotifies.h b/Engine/Plugins/Animation/AnimationWarping/Source/Runtime/Public/EvaluationNotifies/AnimNode_EvaluationNotifies.h
new file mode 100644
--- /dev/null
+++ b/Engine/Plugins/Animation/AnimationWarping/Source/Runtime/Public/EvaluationNotifies/AnimNode_EvaluationNotifies.h
@@ -0,0 +1,75 @@
+// Copyright Epic Games, Inc. All Rights Reserved.
+
+#pragma once
+
+#include "Animation/AnimNotifies/AnimNotifyState.h"
+#include "BoneControllers/BoneControllerTypes.h"
+#include "BoneControllers/AnimNode_SkeletalControlBase.h"
+#include "Kismet/KismetMathLibrary.h"
+#include "StructUtils/InstancedStruct.h"
+#include "AnimNode_EvaluationNotifies.generated.h"
+
+struct FAnimationInitializeContext;
+struct FComponentSpacePoseContext;
+struct FNodeDebugData;
+
+USTRUCT()
+struct FEvaluationNotifyInstance
+{
+	GENERATED_BODY()
+	
+	virtual void Start(const UAnimSequenceBase* AnimationAsset) {}
+	virtual void Update(const UAnimSequenceBase* AnimationAsset, float CurrentTime, float DeltaTime,
+		FTransform& RootBoneTransform, const TMap<FName, FTransform>& NamedTransforms, FComponentSpacePoseContext& Output, TArray<FBoneTransform>& OutBoneTransforms) {}
+	virtual void End() {}
+
+	virtual ~FEvaluationNotifyInstance() {}
+
+	TObjectPtr<UAnimNotifyState> AnimNotify = nullptr;
+	float StartTime = 0.f;
+	float EndTime = 0.f;
+	bool bActive = false;
+};
+
+// add procedural delta to the root motion attribute 
+USTRUCT(BlueprintInternalUseOnly)
+struct ANIMATIONWARPINGRUNTIME_API FAnimNode_EvaluationNotifies : public FAnimNode_SkeletalControlBase
+{
+	GENERATED_BODY()
+
+	// Animation Asset for incorporating root motion data. If CurrentAnimAsset is set, and the animation has root motion rotation within the TargetTime, then those rotations will be scaled to reach the TargetOrientation
+	UPROPERTY(EditAnywhere, Transient, BlueprintReadWrite, Category=Evaluation, meta=(PinShownByDefault))
+	TObjectPtr<UAnimationAsset> CurrentAnimAsset;
+	
+	// Current playback time in seconds of the CurrentAnimAsset
+	UPROPERTY(EditAnywhere, Transient, BlueprintReadWrite, Category=Evaluation, meta=(PinShownByDefault))
+	float CurrentAnimAssetTime = 0.f;
+	
+	// Current playback time in seconds of the CurrentAnimAsset
+	UPROPERTY(EditAnywhere, Transient, BlueprintReadWrite, Category=Evaluation, meta=(PinShownByDefault))
+	TMap<FName, FTransform> NamedTransforms;
+
+	// FAnimNodeBase interface
+	virtual void Initialize_AnyThread(const FAnimationInitializeContext& Context) override;
+	// End of FAnimNodeBase interface
+	
+	// FAnimNode_SkeletalControlBase interface
+	virtual void UpdateInternal(const FAnimationUpdateContext& Context) override;
+	virtual void EvaluateSkeletalControl_AnyThread(FComponentSpacePoseContext& Output, TArray<FBoneTransform>& OutBoneTransforms) override;
+	virtual bool IsValidToEvaluate(const USkeleton* Skeleton, const FBoneContainer& RequiredBones) override { return true; }
+	// End of FAnimNode_SkeletalControlBase interface
+
+	static void RegisterEvaluationHandler(UClass* NotifyType, UScriptStruct* Handler);
+	static void UnregisterEvaluationHandler(UClass* NotifyType);
+
+private:
+	void GetEvaluationNotifiesFromAnimation(const UAnimSequenceBase* Animation, TArray<FInstancedStruct>& OutNotifyInstances);
+
+	FTransform RootBoneTransform;
+	TArray<FInstancedStruct> Tags;
+	float PreviousAnimAssetTime = 0;
+	UPROPERTY(Transient)
+	TObjectPtr<const UAnimSequenceBase> CurrentSequence = nullptr;
+
+	static TMap<UClass*, UScriptStruct*> NotifyEvaluationHandlerMap;
+};
diff --git a/Engine/Plugins/Animation/AnimationWarping/Source/Runtime/Public/EvaluationNotifies/AnimNotifyState_Alignment.h b/Engine/Plugins/Animation/AnimationWarping/Source/Runtime/Public/EvaluationNotifies/AnimNotifyState_Alignment.h
new file mode 100644
--- /dev/null
+++ b/Engine/Plugins/Animation/AnimationWarping/Source/Runtime/Public/EvaluationNotifies/AnimNotifyState_Alignment.h
@@ -0,0 +1,35 @@
+// Copyright Epic Games, Inc. All Rights Reserved.
+
+#pragma once
+
+#include "AnimNode_EvaluationNotifies.h"
+#include "AnimNotifyState_Alignment.generated.h"
+
+UCLASS(BlueprintType, DisplayName="Alignment")
+class UNotifyState_Alignment : public UAnimNotifyState
+{
+	GENERATED_BODY()
+	
+	public:
+
+	UPROPERTY(EditAnywhere, Category = Properties)
+	FBoneReference AlignBone;
+	
+	UPROPERTY(EditAnywhere, Category = Properties)
+	FTransform AlignOffset;
+
+	UPROPERTY(EditAnywhere, Category = Properties)
+	FName TransformName;
+};
+
+USTRUCT()
+struct FAlignmentNotifyInstance : public FEvaluationNotifyInstance
+{
+	GENERATED_BODY()
+
+	virtual void Start(const UAnimSequenceBase* AnimationAsset) override;
+	virtual void Update(const UAnimSequenceBase* AnimationAsset, float CurrentTime, float DeltaTime,
+		FTransform& RootBoneTransform, const TMap<FName, FTransform>& NamedTransforms, FComponentSpacePoseContext& Output, TArray<FBoneTransform>& OutBoneTransforms) override;
+
+	FBoneReference AlignBone;
+};
\ No newline at end of file
