Moving Evaluation notifies in to a separate experimantal plugin and adding a Two Bone IK Evaluation Notify

diff --git a/Engine/Plugins/Animation/AnimationWarping/Source/Editor/Private/AnimGraph/AnimGraphNode_EvaluationNotifies.cpp b/Engine/Plugins/Animation/AnimationWarping/Source/Editor/Private/AnimGraph/AnimGraphNode_EvaluationNotifies.cpp
deleted file mode 100644
diff --git a/Engine/Plugins/Animation/AnimationWarping/Source/Editor/Public/AnimGraph/AnimGraphNode_EvaluationNotifies.h b/Engine/Plugins/Animation/AnimationWarping/Source/Editor/Public/AnimGraph/AnimGraphNode_EvaluationNotifies.h
deleted file mode 100644
diff --git a/Engine/Plugins/Animation/AnimationWarping/Source/Runtime/Private/AnimationWarpingRuntimeModule.cpp b/Engine/Plugins/Animation/AnimationWarping/Source/Runtime/Private/AnimationWarpingRuntimeModule.cpp
--- a/Engine/Plugins/Animation/AnimationWarping/Source/Runtime/Private/AnimationWarpingRuntimeModule.cpp
+++ b/Engine/Plugins/Animation/AnimationWarping/Source/Runtime/Private/AnimationWarpingRuntimeModule.cpp
@@ -7,8 +7,6 @@
 #include "Animation/BuiltInAttributeTypes.h"
 #include "Features/IModularFeatures.h"
 #include "Modules/ModuleManager.h"
-#include "EvaluationNotifies/AnimNode_EvaluationNotifies.h"
-#include "EvaluationNotifies/AnimNotifyState_Alignment.h"
 
 class FAnimationWarpingRuntimeModule : public IAnimationWarpingRuntimeModule
 {
@@ -44,15 +42,11 @@ FModule::FModule()
 void FModule::StartupModule()
 {
 	IModularFeatures::Get().RegisterModularFeature(UE::Anim::IAnimRootMotionProvider::ModularFeatureName, this);
-
-	FAnimNode_EvaluationNotifies::RegisterEvaluationHandler(UNotifyState_Alignment::StaticClass(), FAlignmentNotifyInstance::StaticStruct());
 }
 
 void FModule::ShutdownModule()
 {
 	IModularFeatures::Get().UnregisterModularFeature(UE::Anim::IAnimRootMotionProvider::ModularFeatureName, this);
-	
-	FAnimNode_EvaluationNotifies::UnregisterEvaluationHandler(UNotifyState_Alignment::StaticClass());
 }
 
 void FModule::SampleRootMotion(const FDeltaTimeRecord& SampleRange, const UAnimSequence& Sequence, bool bLoopingSequence, UE::Anim::FStackAttributeContainer& OutAttributes) const
diff --git a/Engine/Plugins/Animation/AnimationWarping/Source/Runtime/Private/EvaluationNotifies/AnimNode_EvaluationNotifies.cpp b/Engine/Plugins/Animation/AnimationWarping/Source/Runtime/Private/EvaluationNotifies/AnimNode_EvaluationNotifies.cpp
deleted file mode 100644
diff --git a/Engine/Plugins/Animation/AnimationWarping/Source/Runtime/Private/EvaluationNotifies/AnimNotifyState_Alignment.cpp b/Engine/Plugins/Animation/AnimationWarping/Source/Runtime/Private/EvaluationNotifies/AnimNotifyState_Alignment.cpp
deleted file mode 100644
diff --git a/Engine/Plugins/Animation/AnimationWarping/Source/Runtime/Public/BoneControllers/AnimNode_OffsetRootBone.h b/Engine/Plugins/Animation/AnimationWarping/Source/Runtime/Public/BoneControllers/AnimNode_OffsetRootBone.h
--- a/Engine/Plugins/Animation/AnimationWarping/Source/Runtime/Public/BoneControllers/AnimNode_OffsetRootBone.h
+++ b/Engine/Plugins/Animation/AnimationWarping/Source/Runtime/Public/BoneControllers/AnimNode_OffsetRootBone.h
@@ -12,7 +12,7 @@ struct FNodeDebugData;
 
 namespace UE::AnimationWarping
 {
-class FRootOffsetProvider : public UE::Anim::IGraphMessage
+class ANIMATIONWARPINGRUNTIME_API FRootOffsetProvider : public UE::Anim::IGraphMessage
 {
 	DECLARE_ANIMGRAPH_MESSAGE(FRootOffsetProvider);
 
diff --git a/Engine/Plugins/Animation/AnimationWarping/Source/Runtime/Public/EvaluationNotifies/AnimNode_EvaluationNotifies.h b/Engine/Plugins/Animation/AnimationWarping/Source/Runtime/Public/EvaluationNotifies/AnimNode_EvaluationNotifies.h
deleted file mode 100644
diff --git a/Engine/Plugins/Animation/AnimationWarping/Source/Runtime/Public/EvaluationNotifies/AnimNotifyState_Alignment.h b/Engine/Plugins/Animation/AnimationWarping/Source/Runtime/Public/EvaluationNotifies/AnimNotifyState_Alignment.h
deleted file mode 100644
