Add CVars to disable steering.

diff --git a/Engine/Plugins/Animation/AnimationWarping/Source/Runtime/Private/BoneControllers/AnimNode_Steering.cpp b/Engine/Plugins/Animation/AnimationWarping/Source/Runtime/Private/BoneControllers/AnimNode_Steering.cpp
--- a/Engine/Plugins/Animation/AnimationWarping/Source/Runtime/Private/BoneControllers/AnimNode_Steering.cpp
+++ b/Engine/Plugins/Animation/AnimationWarping/Source/Runtime/Private/BoneControllers/AnimNode_Steering.cpp
@@ -14,6 +14,13 @@
 
 #include UE_INLINE_GENERATED_CPP_BY_NAME(AnimNode_Steering)
 
+bool bAnimNodeSteeringEnabled = true;
+static FAutoConsoleVariableRef CVarAnimNodeSteeringEnabled(
+	TEXT("a.AnimNode.Steering.Enabled"),
+	bAnimNodeSteeringEnabled,
+	TEXT("True will enable steering anim nodes. Equivalent to setting alpha to non-zero.")
+);
+
 void FAnimNode_Steering::UpdateInternal(const FAnimationUpdateContext& Context)
 {
 	FAnimNode_SkeletalControlBase::UpdateInternal(Context);
@@ -69,7 +76,7 @@ void FAnimNode_Steering::EvaluateSkeletalControl_AnyThread(FComponentSpacePoseCo
 			TargetOrientation = FilteredTarget;
 		}
 		
-		if (Alpha > 0.0f)
+		if (Alpha > 0.0f && bAnimNodeSteeringEnabled)
 		{
 			const UE::Anim::IAnimRootMotionProvider* RootMotionProvider = UE::Anim::IAnimRootMotionProvider::Get();
 			ensureMsgf(RootMotionProvider, TEXT("Steering expected a valid root motion delta provider interface."));
diff --git a/Engine/Plugins/Experimental/AnimNextWarping/Source/AnimNextWarping/Private/SteeringTrait.cpp b/Engine/Plugins/Experimental/AnimNextWarping/Source/AnimNextWarping/Private/SteeringTrait.cpp
--- a/Engine/Plugins/Experimental/AnimNextWarping/Source/AnimNextWarping/Private/SteeringTrait.cpp
+++ b/Engine/Plugins/Experimental/AnimNextWarping/Source/AnimNextWarping/Private/SteeringTrait.cpp
@@ -14,6 +14,13 @@
 #include "TraitInterfaces/ITimeline.h"
 #include "VisualLogger/VisualLogger.h"
 
+bool bAnimNextSteeringTraitEnabled = true;
+static FAutoConsoleVariableRef CVarAnimNextSteeringTraitEnabled(
+	TEXT("a.AnimNext.SteeringTrait.Enabled"),
+	bAnimNextSteeringTraitEnabled,
+	TEXT("True will enable steering for AnimNext. Equivalent to setting alpha to non-zero.")
+);
+
 ////////////////////////////////////////////////////////////////////////////////////////////////////////////////////////
 // FSteeringTrait
 
@@ -123,7 +130,7 @@ void FAnimNextSteeringTask::Execute(UE::AnimNext::FEvaluationVM& VM) const
 			InstanceData->TargetOrientation = InstanceData->FilteredTarget;
 		}
 
-		if (SharedData->Alpha > 0.0f)
+		if (SharedData->Alpha > 0.0f && bAnimNextSteeringTraitEnabled)
 		{
 			const UE::Anim::IAnimRootMotionProvider* RootMotionProvider = UE::Anim::IAnimRootMotionProvider::Get();
 			if (!RootMotionProvider)
