Fix Non-Normalized rotations coming out of OffsetRootBone node

diff --git a/Engine/Plugins/Animation/AnimationWarping/Source/Runtime/Private/BoneControllers/AnimNode_OffsetRootBone.cpp b/Engine/Plugins/Animation/AnimationWarping/Source/Runtime/Private/BoneControllers/AnimNode_OffsetRootBone.cpp
--- a/Engine/Plugins/Animation/AnimationWarping/Source/Runtime/Private/BoneControllers/AnimNode_OffsetRootBone.cpp
+++ b/Engine/Plugins/Animation/AnimationWarping/Source/Runtime/Private/BoneControllers/AnimNode_OffsetRootBone.cpp
@@ -170,6 +170,9 @@ void FAnimNode_OffsetRootBone::Evaluate_AnyThread(FPoseContext& Output)
 		// Apply the offset as is (component space) in manual mode
 		RootMotionTransformDelta = FTransform(GetRotationDelta(), GetTranslationDelta());
 	}
+
+	RootMotionTransformDelta.NormalizeRotation();
+	
 	float MaxTranslationOffset = GetMaxTranslationError();
 	
 	bool bCollisionDetected = false;
