Fix steering behaviour for non-looping animations

diff --git a/Engine/Plugins/Animation/AnimationWarping/Source/Runtime/Private/BoneControllers/AnimNode_Steering.cpp b/Engine/Plugins/Animation/AnimationWarping/Source/Runtime/Private/BoneControllers/AnimNode_Steering.cpp
--- a/Engine/Plugins/Animation/AnimationWarping/Source/Runtime/Private/BoneControllers/AnimNode_Steering.cpp
+++ b/Engine/Plugins/Animation/AnimationWarping/Source/Runtime/Private/BoneControllers/AnimNode_Steering.cpp
@@ -94,11 +94,11 @@ void FAnimNode_Steering::EvaluateSkeletalControl_AnyThread(FComponentSpacePoseCo
 								FTransform PredictedRootMotionTransform;
 								if (bMirrored && MirrorDataTable)
 								{
-									PredictedRootMotionTransform = UE::Anim::ExtractRootMotionFromAnimationAsset(AnimSequence, MirrorDataTable, CurrentAnimAssetTime, AnimatedTargetTime, true);
+									PredictedRootMotionTransform = UE::Anim::ExtractRootMotionFromAnimationAsset(AnimSequence, MirrorDataTable, CurrentAnimAssetTime, AnimatedTargetTime, AnimSequence->bLoop);
 								}
 								else
 								{
-									const FAnimExtractContext Context(static_cast<double>(CurrentAnimAssetTime), true, FDeltaTimeRecord(AnimatedTargetTime), true);
+									const FAnimExtractContext Context(static_cast<double>(CurrentAnimAssetTime), true, FDeltaTimeRecord(AnimatedTargetTime), AnimSequence->bLoop);
 									PredictedRootMotionTransform = AnimSequence->ExtractRootMotion(Context);
 								}
 
diff --git a/Engine/Plugins/Experimental/AnimNextAnimGraph/Source/AnimNextAnimGraph/Private/Traits/BlendStackTrait.cpp b/Engine/Plugins/Experimental/AnimNextAnimGraph/Source/AnimNextAnimGraph/Private/Traits/BlendStackTrait.cpp
--- a/Engine/Plugins/Experimental/AnimNextAnimGraph/Source/AnimNextAnimGraph/Private/Traits/BlendStackTrait.cpp
+++ b/Engine/Plugins/Experimental/AnimNextAnimGraph/Source/AnimNextAnimGraph/Private/Traits/BlendStackTrait.cpp
@@ -334,7 +334,7 @@ namespace UE::AnimNext
 						{
 							// We do not check for lifetimes, assume the sequence is alive during pose list execution.
 							check(AnimSequence->IsValidLowLevel());
-							return AnimSequence->ExtractRootMotion(FAnimExtractContext(static_cast<double>(StartTime), true, FDeltaTimeRecord(DeltaTime), bAllowLooping)); 
+							return AnimSequence->ExtractRootMotion(FAnimExtractContext(static_cast<double>(StartTime), true, FDeltaTimeRecord(DeltaTime), bAllowLooping && AnimSequence->bLoop)); 
 						};
 
 						return FOnExtractRootMotionAttribute::CreateLambda(ExtractRootMotionAttribute);
diff --git a/Engine/Plugins/Experimental/AnimNextAnimGraph/Source/AnimNextAnimGraph/Private/Traits/SequencePlayer.cpp b/Engine/Plugins/Experimental/AnimNextAnimGraph/Source/AnimNextAnimGraph/Private/Traits/SequencePlayer.cpp
--- a/Engine/Plugins/Experimental/AnimNextAnimGraph/Source/AnimNextAnimGraph/Private/Traits/SequencePlayer.cpp
+++ b/Engine/Plugins/Experimental/AnimNextAnimGraph/Source/AnimNextAnimGraph/Private/Traits/SequencePlayer.cpp
@@ -65,7 +65,7 @@ namespace UE::AnimNext
 			{
 				// We do not check for lifetimes, assume the sequence is alive during pose list exectuion.
 				check(AnimSequence->IsValidLowLevel());
-				return AnimSequence->ExtractRootMotion(FAnimExtractContext(static_cast<double>(StartTime), true, FDeltaTimeRecord(DeltaTime), bAllowLooping));
+				return AnimSequence->ExtractRootMotion(FAnimExtractContext(static_cast<double>(StartTime), true, FDeltaTimeRecord(DeltaTime), bAllowLooping && AnimSequence->bLoop));
 			};
 
 			return FOnExtractRootMotionAttribute::CreateLambda(ExtractRootMotionAttribute);
