AnimNode_OffsetRootBone - Moved EOffsetRootBoneMode enum to a public header file so it can be referenced from native code on the game side

diff --git a/Engine/Plugins/Animation/AnimationWarping/Source/Runtime/Public/AnimationWarpingTypes.h b/Engine/Plugins/Animation/AnimationWarping/Source/Runtime/Public/AnimationWarpingTypes.h
new file mode 100644
--- /dev/null
+++ b/Engine/Plugins/Animation/AnimationWarping/Source/Runtime/Public/AnimationWarpingTypes.h
@@ -0,0 +1,31 @@
+// Copyright Epic Games, Inc. All Rights Reserved.
+
+#pragma once
+
+#include "CoreMinimal.h"
+#include "AnimationWarpingTypes.generated.h"
+
+UENUM(BlueprintType)
+enum class EOffsetRootBoneMode : uint8
+{
+	// Accumulate the mesh component's movement into the offset.
+	// In this mode, if the mesh component moves 
+	// the offset will counter the motion, and the root will stay in place
+	Accumulate,
+	// Continuously interpolate the offset out
+	// In this mode, if the mesh component moves
+	// The root will stay behind, but will attempt to catch up
+	Interpolate,
+	// Stops accumulating the mesh component's movement delta into the root offset
+	// In this mode, whatever offset we have will be locked but we will still consume animated root motion 
+	LockOffsetAndConsumeAnimation,
+	// Stops accumulating the mesh component's movement delta into the root offset
+	// In this mode, whatever offset we have will be locked but we will still consume animated root motion, as long as it's decreasing the offset.
+	LockOffsetIncreaseAndConsumeAnimation,
+	// Stops accumulating the mesh component's movement delta into the root offset
+	// In this mode, whatever offset we have will be locked and we will ignore animated root motion
+	LockOffsetAndIgnoreAnimation,
+	// Release the offset and stop accumulating the mesh component's movement delta.
+	// In this mode we will "blend out" the offset
+	Release,
+};
diff --git a/Engine/Plugins/Animation/AnimationWarping/Source/Runtime/Public/BoneControllers/AnimNode_OffsetRootBone.h b/Engine/Plugins/Animation/AnimationWarping/Source/Runtime/Public/BoneControllers/AnimNode_OffsetRootBone.h
--- a/Engine/Plugins/Animation/AnimationWarping/Source/Runtime/Public/BoneControllers/AnimNode_OffsetRootBone.h
+++ b/Engine/Plugins/Animation/AnimationWarping/Source/Runtime/Public/BoneControllers/AnimNode_OffsetRootBone.h
@@ -4,6 +4,7 @@
 
 #include "BoneControllers/BoneControllerTypes.h"
 #include "BoneControllers/AnimNode_SkeletalControlBase.h"
+#include "AnimationWarpingTypes.h"
 #include "AnimNode_OffsetRootBone.generated.h"
 
 #define UE_API ANIMATIONWARPINGRUNTIME_API
@@ -32,31 +33,6 @@ private:
 };
 }
 
-UENUM(BlueprintType)
-enum class EOffsetRootBoneMode : uint8
-{
-	// Accumulate the mesh component's movement into the offset.
-	// In this mode, if the mesh component moves 
-	// the offset will counter the motion, and the root will stay in place
-	Accumulate,
-	// Continuously interpolate the offset out
-	// In this mode, if the mesh component moves
-	// The root will stay behind, but will attempt to catch up
-	Interpolate,
-	// Stops accumulating the mesh component's movement delta into the root offset
-	// In this mode, whatever offset we have will be locked but we will still consume animated root motion 
-	LockOffsetAndConsumeAnimation,
-	// Stops accumulating the mesh component's movement delta into the root offset
-	// In this mode, whatever offset we have will be locked but we will still consume animated root motion, as long as it's decreasing the offset.
-	LockOffsetIncreaseAndConsumeAnimation,
-	// Stops accumulating the mesh component's movement delta into the root offset
-	// In this mode, whatever offset we have will be locked and we will ignore animated root motion
-	LockOffsetAndIgnoreAnimation,
-	// Release the offset and stop accumulating the mesh component's movement delta.
-	// In this mode we will "blend out" the offset
-	Release,
-};
-
 UENUM(BlueprintType)
 enum class EOffsetRootBone_CollisionTestingMode : uint8
 {
