On behalf of roland.munguia.

Add threadsafe method for getting curve asset float & vector values.

#tests editor, ABP

diff --git a/Engine/Plugins/Animation/AnimationWarping/Source/Runtime/Private/AnimationWarpingLibrary.cpp b/Engine/Plugins/Animation/AnimationWarping/Source/Runtime/Private/AnimationWarpingLibrary.cpp
--- a/Engine/Plugins/Animation/AnimationWarping/Source/Runtime/Private/AnimationWarpingLibrary.cpp
+++ b/Engine/Plugins/Animation/AnimationWarping/Source/Runtime/Private/AnimationWarpingLibrary.cpp
@@ -3,6 +3,8 @@
 #include "AnimationWarpingLibrary.h"
 #include "BoneControllers/AnimNode_OffsetRootBone.h"
 #include "Animation/AnimMontage.h"
+#include "Curves/CurveFloat.h"
+#include "Curves/CurveVector.h"
 
 #include UE_INLINE_GENERATED_CPP_BY_NAME(AnimationWarpingLibrary)
 
@@ -55,3 +57,25 @@ bool UAnimationWarpingLibrary::GetCurveValueFromAnimation(const UAnimSequenceBas
 
 	return false;
 }
+
+bool UAnimationWarpingLibrary::GetFloatValueFromCurve(const UCurveFloat* InCurve, float InTime, float& OutValue)
+{
+	if (InCurve != nullptr)
+	{
+		OutValue = InCurve->GetFloatValue(InTime);
+		return true;
+	}
+	
+	return false;
+}
+
+bool UAnimationWarpingLibrary::GetVectorValueFromCurve(const UCurveVector* InCurve, float InTime, FVector& OutValue)
+{
+	if (InCurve != nullptr)
+	{
+		OutValue = InCurve->GetVectorValue(InTime);
+		return true;
+	}
+	
+	return false;
+}
\ No newline at end of file
diff --git a/Engine/Plugins/Animation/AnimationWarping/Source/Runtime/Public/AnimationWarpingLibrary.h b/Engine/Plugins/Animation/AnimationWarping/Source/Runtime/Public/AnimationWarpingLibrary.h
--- a/Engine/Plugins/Animation/AnimationWarping/Source/Runtime/Public/AnimationWarpingLibrary.h
+++ b/Engine/Plugins/Animation/AnimationWarping/Source/Runtime/Public/AnimationWarpingLibrary.h
@@ -10,6 +10,8 @@
 #define UE_API ANIMATIONWARPINGRUNTIME_API
 
 struct FAnimNode_OffsetRootBone;
+class UCurveFloat;
+class UCurveVector;
 
 // Exposes operations related to Animation Warping
 UCLASS(MinimalAPI, Experimental)
@@ -25,6 +27,14 @@ public:
 	/** Helper function to extract the value of a curve in an animation at a given time */
 	UFUNCTION(BlueprintPure, Category = "Animation", meta = (BlueprintThreadSafe))
 	static UE_API bool GetCurveValueFromAnimation(const UAnimSequenceBase* Animation, FName CurveName, float Time, float& OutValue);
+
+	/** Helper function to extract the float value from a curve asset at a given time. To be used inside AnimationBlueprint and assumes AnimBP holds reference to CurveAsset for it to be valid. */
+	UFUNCTION(BlueprintPure, Category = "Animation", meta = (BlueprintThreadSafe))
+	static bool GetFloatValueFromCurve(const UCurveFloat* InCurve, float Time, float& OutValue);
+
+	/** Helper function to extract the vector value from a curve asset at a given time. To be used inside AnimationBlueprint and assumes AnimBP holds reference to CurveAsset for it to be valid. */
+	UFUNCTION(BlueprintPure, Category = "Animation", meta = (BlueprintThreadSafe))
+	static bool GetVectorValueFromCurve(const UCurveVector* InCurve, float Time, FVector& OutValue);
 };
 
 #undef UE_API
