OffsetRootBone: Added support for teleporting

#Synced-CL 48507157

diff --git a/Engine/Plugins/Animation/AnimationWarping/Source/Runtime/Private/BoneControllers/AnimNode_OffsetRootBone.cpp b/Engine/Plugins/Animation/AnimationWarping/Source/Runtime/Private/BoneControllers/AnimNode_OffsetRootBone.cpp
--- a/Engine/Plugins/Animation/AnimationWarping/Source/Runtime/Private/BoneControllers/AnimNode_OffsetRootBone.cpp
+++ b/Engine/Plugins/Animation/AnimationWarping/Source/Runtime/Private/BoneControllers/AnimNode_OffsetRootBone.cpp
@@ -145,9 +145,32 @@ void FAnimNode_OffsetRootBone::Evaluate_AnyThread(FPoseContext& Output)
 	const FCompactPoseBoneIndex TargetBoneIndex(0);
 	const FTransform InputBoneTransform = Output.Pose[TargetBoneIndex];
 
-	const FTransform LastComponentTransform = ComponentTransform;
+	FTransform LastComponentTransform = ComponentTransform;
 	ComponentTransform = AnimInstanceProxy->GetComponentTransform();
 
+	// If the difference between the previous and current component transform is above the teleport threshold
+	// then we adjust the simulated position to account for the teleport. Then we set the last transform to be 
+	// the same as the previous so that the node does not attempt to counter-animate out the component transform 
+	// this frame
+	const float TeleportDistanceThreshold = Output.AnimInstanceProxy->GetSkelMeshComponent()->GetTeleportDistanceThreshold();
+	if (FVector::Distance(LastComponentTransform.GetLocation(), ComponentTransform.GetLocation()) > TeleportDistanceThreshold)
+	{
+		if (GetResetOnTeleport())
+		{
+			SimulatedTranslation = ComponentTransform.GetTranslation();
+			SimulatedRotation = ComponentTransform.GetRotation();
+		}
+		else
+		{
+			const FTransform OffsetTransform = FTransform(SimulatedRotation, SimulatedTranslation).GetRelativeTransformReverse(LastComponentTransform);
+			const FTransform TeleportedTransform = OffsetTransform * ComponentTransform;
+			SimulatedTranslation = TeleportedTransform.GetLocation();
+			SimulatedRotation = TeleportedTransform.GetRotation();
+		}
+
+		LastComponentTransform = ComponentTransform;
+	}
+
 	const EOffsetRootBoneMode CurrentTranslationMode = GetTranslationMode();
 	const EOffsetRootBoneMode CurrentRotationMode = GetRotationMode();
 
@@ -596,6 +619,11 @@ const FVector& FAnimNode_OffsetRootBone::GetCollisionTestShapeOffset() const
 	return GET_ANIM_NODE_DATA(FVector, CollisionTestShapeOffset);
 }
 
+bool FAnimNode_OffsetRootBone::GetResetOnTeleport() const
+{
+	return GET_ANIM_NODE_DATA(bool, bResetOnTeleport);
+}
+
 void FAnimNode_OffsetRootBone::Reset(const FAnimationBaseContext& Context)
 {
 	ComponentTransform = Context.AnimInstanceProxy->GetComponentTransform();
diff --git a/Engine/Plugins/Animation/AnimationWarping/Source/Runtime/Public/BoneControllers/AnimNode_OffsetRootBone.h b/Engine/Plugins/Animation/AnimationWarping/Source/Runtime/Public/BoneControllers/AnimNode_OffsetRootBone.h
--- a/Engine/Plugins/Animation/AnimationWarping/Source/Runtime/Public/BoneControllers/AnimNode_OffsetRootBone.h
+++ b/Engine/Plugins/Animation/AnimationWarping/Source/Runtime/Public/BoneControllers/AnimNode_OffsetRootBone.h
@@ -107,7 +107,7 @@ struct FAnimNode_OffsetRootBone : public FAnimNode_Base
 	UPROPERTY(EditAnywhere, Category = Settings, meta = (EditCondition = bClampToRotationVelocity, FoldProperty, PinHiddenByDefault))
 	float RotationSpeedRatio = 0.5f;
 
-	// When OnGround is true, root motion velicities will be projected onto the ground surface
+	// When OnGround is true, root motion velocities will be projected onto the ground surface
 	UPROPERTY(EditAnywhere, Category = Settings, meta = (FoldProperty, PinHiddenByDefault))
 	bool bOnGround = true;
 	
@@ -125,6 +125,9 @@ struct FAnimNode_OffsetRootBone : public FAnimNode_Base
 	UPROPERTY(EditAnywhere, Category = Evaluation, meta = (FoldProperty, PinHiddenByDefault))
 	FRotator RotationDelta = FRotator::ZeroRotator;
 
+	// If to reset the root bone offset on teleport
+	UPROPERTY(EditAnywhere, Category = Settings, meta = (FoldProperty, PinHiddenByDefault))
+	bool bResetOnTeleport = false;
 
 	UPROPERTY(EditAnywhere, Category = CollisionTesting, meta = (FoldProperty, PinHiddenByDefault))
 	EOffsetRootBone_CollisionTestingMode CollisionTestingMode = EOffsetRootBone_CollisionTestingMode::Disabled;
@@ -164,6 +167,7 @@ public:
 	UE_API EOffsetRootBone_CollisionTestingMode GetCollisionTestingMode() const;
 	UE_API float GetCollisionTestShapeRadius() const;
 	UE_API const FVector& GetCollisionTestShapeOffset() const;
+	UE_API bool GetResetOnTeleport() const;
 
 	// get the current simulated root transform
 	void GetOffsetRootTransform(FTransform& OutTransform)
diff --git a/Engine/Plugins/Experimental/UAF/UAFWarping/Source/UAFWarping/Private/OffsetRootBoneTrait.cpp b/Engine/Plugins/Experimental/UAF/UAFWarping/Source/UAFWarping/Private/OffsetRootBoneTrait.cpp
--- a/Engine/Plugins/Experimental/UAF/UAFWarping/Source/UAFWarping/Private/OffsetRootBoneTrait.cpp
+++ b/Engine/Plugins/Experimental/UAF/UAFWarping/Source/UAFWarping/Private/OffsetRootBoneTrait.cpp
@@ -184,9 +184,32 @@ void FAnimNextOffsetRootBoneTask::Execute(UE::UAF::FEvaluationVM& VM) const
 	}
 
 	// Note we set last to current on first update
-	const FTransform LastComponentTransform = InstanceData->bIsFirstUpdate ? ComponentTransform : InstanceData->LastMeshComponentTransformWorld;
+	FTransform LastComponentTransform = InstanceData->bIsFirstUpdate ? ComponentTransform : InstanceData->LastMeshComponentTransformWorld;
 	InstanceData->bIsFirstUpdate = false;
 	
+
+	// If the difference between the previous and current component transform is above the teleport threshold
+	// then we adjust the simulated position to account for the teleport. Then we set the last transform to be 
+	// the same as the previous so that the node does not attempt to counter-animate out the component transform 
+	// this frame
+	if (FVector::Distance(LastComponentTransform.GetLocation(), ComponentTransform.GetLocation()) > SharedData->TeleportDistanceThreshold)
+	{
+		if (SharedData->bResetOnTeleport)
+		{
+			InstanceData->SimulatedTranslation = ComponentTransform.GetTranslation();
+			InstanceData->SimulatedRotation = ComponentTransform.GetRotation();
+		}
+		else
+		{
+			const FTransform OffsetTransform = FTransform(InstanceData->SimulatedRotation, InstanceData->SimulatedTranslation).GetRelativeTransformReverse(LastComponentTransform);
+			const FTransform TeleportedTransform = OffsetTransform * ComponentTransform;
+			InstanceData->SimulatedTranslation = TeleportedTransform.GetLocation();
+			InstanceData->SimulatedRotation = TeleportedTransform.GetRotation();
+		}
+
+		LastComponentTransform = ComponentTransform;
+	}
+
 	// Cycle last transform
 	InstanceData->LastMeshComponentTransformWorld = InstanceData->MeshComponentTransformWorld;
 
diff --git a/Engine/Plugins/Experimental/UAF/UAFWarping/Source/UAFWarping/Private/OffsetRootBoneTrait.h b/Engine/Plugins/Experimental/UAF/UAFWarping/Source/UAFWarping/Private/OffsetRootBoneTrait.h
--- a/Engine/Plugins/Experimental/UAF/UAFWarping/Source/UAFWarping/Private/OffsetRootBoneTrait.h
+++ b/Engine/Plugins/Experimental/UAF/UAFWarping/Source/UAFWarping/Private/OffsetRootBoneTrait.h
@@ -103,6 +103,14 @@ struct FOffsetRootBoneTraitSharedData : public FAnimNextTraitSharedData
 	UPROPERTY(EditAnywhere, Category = Settings)
 	FVector AnimatedGroundNormal = { 0, 0, 1 };
 
+	// Distance threshold used to detect teleports
+	UPROPERTY(EditAnywhere, Category = Settings)
+	float TeleportDistanceThreshold = 300.0f;
+
+	// If to reset the root bone offset on teleport
+	UPROPERTY(EditAnywhere, Category = Settings)
+	bool bResetOnTeleport = false;
+
 	// Latent pin support boilerplate
 	#define TRAIT_LATENT_PROPERTIES_ENUMERATOR(GeneratorMacro) \
 		GeneratorMacro(Alpha) \
