Fix for potential NaN bone transforms generated inside AnimNode_FootPlacement if root is rotated exactly 90 degrees from expected approach direction
Add extra nan checks

#Synced-CL 50205727

diff --git a/Engine/Plugins/Animation/AnimationWarping/Source/Runtime/Private/BoneControllers/AnimNode_FootPlacement.cpp b/Engine/Plugins/Animation/AnimationWarping/Source/Runtime/Private/BoneControllers/AnimNode_FootPlacement.cpp
--- a/Engine/Plugins/Animation/AnimationWarping/Source/Runtime/Private/BoneControllers/AnimNode_FootPlacement.cpp
+++ b/Engine/Plugins/Animation/AnimationWarping/Source/Runtime/Private/BoneControllers/AnimNode_FootPlacement.cpp
@@ -230,8 +230,14 @@ namespace UE::Anim::FootPlacement
 		return Centroid;
 	}
 
-	static float GetDistanceToPlaneAlongDirection(const FVector& Location, const FPlane& PlantPlane, const FVector& ApproachDir)
+	static TOptional<float> GetDistanceToPlaneAlongDirection(const FVector& Location, const FPlane& PlantPlane, const FVector& ApproachDir)
 	{
+		// If approach dir is perpendicular to plant plane, distance is undefined
+		if (FMath::IsNearlyZero(ApproachDir | PlantPlane))
+		{
+			return TOptional<float>();
+		}
+		
 		const FVector IntersectionLoc = UE::Anim::FootPlacement::PointDirectionPlaneIntersection(
 			Location,
 			-ApproachDir,
@@ -395,12 +401,12 @@ float FAnimNode_FootPlacement::CalcTargetPlantPlaneDistance(
 				PelvisData.InputPose.IKRootTransformCS.TransformVectorNoScale(FVector::UpVector));
 
 	// TODO: I'm just getting the distance between bones and the plane, instead of actual foot/ball bases
-	const float FootBaseDistance = 
+	const TOptional<float> FootBaseDistance = 
 		UE::Anim::FootPlacement::GetDistanceToPlaneAlongDirection(LegInputPose.FootTransformCS.GetLocation(), IKGroundPlaneCS, Context.ApproachDirCS);
-	const float BallBaseDistance = 
+	const TOptional<float> BallBaseDistance = 
 		UE::Anim::FootPlacement::GetDistanceToPlaneAlongDirection(IKBallBoneCS.GetLocation(), IKGroundPlaneCS, Context.ApproachDirCS);
 
-	const float PlantPlaneDistance = FMath::Min(FootBaseDistance, BallBaseDistance);
+	const float PlantPlaneDistance = FMath::Min(FootBaseDistance.Get(0.0f), BallBaseDistance.Get(0.0f));
 	return PlantPlaneDistance;
 }
 
@@ -419,7 +425,7 @@ void FAnimNode_FootPlacement::AlignPlantToGround(
 	// TODO: This wont work well when the animation doesn't have a single plant plane, i.e. a walking upstairs anim
 	const FTransform IKFootRootWS = PelvisData.InputPose.IKRootTransformCS * Context.OwningComponentToWorld;
 	const FPlane IKFootRootPlaneWS = FPlane(IKFootRootWS.GetLocation(), IKFootRootWS.TransformVectorNoScale(FVector::UpVector));
-	const float IKFootRootToFootRootTargetDistance = 
+	const TOptional<float> IKFootRootToFootRootTargetDistance = 
 		UE::Anim::FootPlacement::GetDistanceToPlaneAlongDirection(InputPoseFootTransformWS.GetLocation(), IKFootRootPlaneWS, Context.ApproachDirWS);
 
 	const FVector CorrectedPlaneIntersectionWS = UE::Anim::FootPlacement::PointDirectionPlaneIntersection(
@@ -428,7 +434,7 @@ void FAnimNode_FootPlacement::AlignPlantToGround(
 		PlantPlaneWS);
 
 	const FVector CorrectedLocationWS =
-		CorrectedPlaneIntersectionWS - (Context.ApproachDirWS * IKFootRootToFootRootTargetDistance);
+		CorrectedPlaneIntersectionWS - (Context.ApproachDirWS * IKFootRootToFootRootTargetDistance.Get(0.0f));
 
 	// The relationship between the ik reference and the normal of the plane defined by the ik foot root must also be 
 	// respected
@@ -568,9 +574,9 @@ void FAnimNode_FootPlacement::UpdatePlantingPlaneInterpolation(
 		{
 			// Prevent the foot from clipping too much into geometry due to interpolation
 			const FPlane GroundPlane = FPlane(ImpactLocationWS, InOutPlantPlane.GetNormal());
-			const float DistanceToGroundPlane =
+			const TOptional<float> DistanceToGroundPlane =
 				UE::Anim::FootPlacement::GetDistanceToPlaneAlongDirection(CurrPlaneIntersection, GroundPlane, Context.ApproachDirCS);
-			const float PenetrationAmount = -DistanceToGroundPlane - TraceSettings.MaxGroundPenetration;
+			const float PenetrationAmount = -DistanceToGroundPlane.Get(0.0f) - TraceSettings.MaxGroundPenetration;
 			if (PenetrationAmount > 0.0f)
 			{
 				CurrPlaneIntersection -= Context.ApproachDirCS * PenetrationAmount;
@@ -817,6 +823,7 @@ UE::Anim::FootPlacement::FPlantResult FAnimNode_FootPlacement::FinalizeFootAlign
 
 					// Move the foot bone closer to the hip prevent overextension
 					CorrectedFootTransformCS.SetLocation(CorrectedFootLocationCS);
+					check(!CorrectedFootTransformCS.ContainsNaN());
 				}
 
 				// Fix any remaining hyper-extension
@@ -828,6 +835,7 @@ UE::Anim::FootPlacement::FPlantResult FAnimNode_FootPlacement::FinalizeFootAlign
 					FVector NotHyperextendedPlantLocation;
 					FMath::SphereDistToLine(FinalHipTransformCS.GetLocation(), MaxExtension, CorrectedFootTransformCS.GetLocation(), CorrectedHipToFootDir, NotHyperextendedPlantLocation);
 					CorrectedFootTransformCS.SetLocation(NotHyperextendedPlantLocation);
+					check(!CorrectedFootTransformCS.ContainsNaN());
 				}
 			}
 			else
@@ -849,15 +857,15 @@ UE::Anim::FootPlacement::FPlantResult FAnimNode_FootPlacement::FinalizeFootAlign
 	// with the geometry).
 	// TODO: Consolidate with CalcTargetPlantPlaneDistance
 	const FPlane PlantPlaneCS = LegData.Plant.GetPlantPlaneCS(GetRootToComponent());
-	const float FootDistance = UE::Anim::FootPlacement::GetDistanceToPlaneAlongDirection(
+	const TOptional<float> FootDistance = UE::Anim::FootPlacement::GetDistanceToPlaneAlongDirection(
 		CorrectedBallLocationCS,
 		PlantPlaneCS,
 		Context.ApproachDirCS);
-	const float BallDistance = UE::Anim::FootPlacement::GetDistanceToPlaneAlongDirection(
+	const TOptional<float> BallDistance = UE::Anim::FootPlacement::GetDistanceToPlaneAlongDirection(
 		CorrectedFootTransformCS.GetLocation(),
 		PlantPlaneCS,
 		Context.ApproachDirCS);
-	float MinDistance = FMath::Min(FootDistance, BallDistance);
+	float MinDistance = FMath::Min(FootDistance.Get(0.0f), BallDistance.Get(0.0f));
 	// Allow as much penetration as the source animation.
 	MinDistance -= FMath::Min(0.0f, LegData.InputPose.DistanceToPlant);
 	// A min distance < 0.0f means there was penetration
@@ -865,6 +873,7 @@ UE::Anim::FootPlacement::FPlantResult FAnimNode_FootPlacement::FinalizeFootAlign
 	{
 		CorrectedFootTransformCS.AddToTranslation(MinDistance * Context.ApproachDirCS);
 	}
+	check(!CorrectedFootTransformCS.ContainsNaN());
 
 	// Fix any remaining hyper-compression. Clip into the ground plane if necessary.
 	// Doing this after pushing the feet out of the ground plane ensures we won't end up in awkward poses. 
@@ -1798,15 +1807,15 @@ void FAnimNode_FootPlacement::ProcessFootAlignment(
 		const FVector PlaneCenter = PelvisData.InputPose.FootMidpointCS + PlaneNormal * PlantSettings.SeparatingDistance;
 		const FPlane SeparatingPlane = FPlane(PlaneCenter, PlaneNormal);
 		
-		const float DistanceToSeparatingPlane =
+		const TOptional<float> DistanceToSeparatingPlane =
 			UE::Anim::FootPlacement::GetDistanceToPlaneAlongDirection(FootUnalignedLocationRS, SeparatingPlane, -PlaneNormal);
 
 		if (Plant.PlantType == UE::Anim::FootPlacement::EPlantType::Unplanted)
 		{
 			FVector SeparatingPlaneOffset(FVector::ZeroVector);
-			if (DistanceToSeparatingPlane < 0.0f)
+			if (DistanceToSeparatingPlane.Get(0.0f) < 0.0f)
 			{
-				SeparatingPlaneOffset = -PlaneNormal * DistanceToSeparatingPlane ;
+				SeparatingPlaneOffset = -PlaneNormal * DistanceToSeparatingPlane.GetValue() ;
 			}
 
 			if (InterpolationSettings.bEnableSeparationInterpolation)
@@ -1840,7 +1849,7 @@ void FAnimNode_FootPlacement::ProcessFootAlignment(
 				PlaneCenterWS, 25.0f, 24, FColor::Red,
 				PlaneNormalWS, false, -1.0f, SDPG_Foreground, 0.5f);
 
-			DebugData.LegsInfo[LegData.Idx].DistanceToSeparatingPlane = DistanceToSeparatingPlane;
+			DebugData.LegsInfo[LegData.Idx].DistanceToSeparatingPlane = DistanceToSeparatingPlane.Get(0.0f);
 		}
 #endif
 		
