Speculative fix for NaN bone transforms coming from AnimNode_StrideWarping

#Synced-CL 50205727

diff --git a/Engine/Plugins/Animation/AnimationWarping/Source/Runtime/Private/BoneControllers/AnimNode_StrideWarping.cpp b/Engine/Plugins/Animation/AnimationWarping/Source/Runtime/Private/BoneControllers/AnimNode_StrideWarping.cpp
--- a/Engine/Plugins/Animation/AnimationWarping/Source/Runtime/Private/BoneControllers/AnimNode_StrideWarping.cpp
+++ b/Engine/Plugins/Animation/AnimationWarping/Source/Runtime/Private/BoneControllers/AnimNode_StrideWarping.cpp
@@ -84,6 +84,22 @@ void FAnimNode_StrideWarping::EvaluateSkeletalControl_AnyThread(FComponentSpaceP
 		bGraphDrivenWarping = !!RootMotionProvider;
 		ensureMsgf(bGraphDrivenWarping, TEXT("Graph driven Stride Warping expected a valid root motion delta provider interface."));
 	}
+	
+	if (!bGraphDrivenWarping)
+	{
+		// If we aren't using graph driven warping, we need to validate the input variables are valid
+		if (StrideDirection.ContainsNaN() || StrideDirection.IsNearlyZero())
+		{
+			UE_LOG(LogAnimation, Error, TEXT("Invalid stride direction in FAnimNode_StrideWarping"));
+			return;
+		}
+
+		if (!FMath::IsFinite(StrideScale))
+		{
+			UE_LOG(LogAnimation, Error, TEXT("Invalid stride scale in FAnimNode_StrideWarping"));
+			return;
+		}
+	}
 
 	FTransform RootMotionTransformDelta = FTransform::Identity;
 #if WITH_EDITORONLY_DATA
@@ -161,12 +177,12 @@ void FAnimNode_StrideWarping::EvaluateSkeletalControl_AnyThread(FComponentSpaceP
 	if (bGraphDrivenWarping)
 	{
 #if WITH_EDITORONLY_DATA
-		CachedRootMotionDeltaSpeed = CachedRootMotionDeltaTranslation.Size() / CachedDeltaTime;
+		CachedRootMotionDeltaSpeed = !FMath::IsNearlyZero(CachedDeltaTime) ? (CachedRootMotionDeltaTranslation.Size() / CachedDeltaTime) : 0.f;
 #else
-		const float CachedRootMotionDeltaSpeed = CachedRootMotionDeltaTranslation.Size() / CachedDeltaTime;
+		const float CachedRootMotionDeltaSpeed = !FMath::IsNearlyZero(CachedDeltaTime) ? (CachedRootMotionDeltaTranslation.Size() / CachedDeltaTime) : 0.f;
 #endif
 
-		if (CachedRootMotionDeltaSpeed <= MinRootMotionSpeedThreshold)
+		if (CachedRootMotionDeltaSpeed <= MinRootMotionSpeedThreshold || FMath::IsNearlyZero(CachedRootMotionDeltaSpeed))
 		{
 			// If root motion speed is under the threshold, snap back to no stride adjustment.
 			// If interpolation is on, it will blend the result
@@ -182,6 +198,7 @@ void FAnimNode_StrideWarping::EvaluateSkeletalControl_AnyThread(FComponentSpaceP
 
 	// Allow the opportunity for stride scale clamping and interpolation regardless of evaluation mode
 	ActualStrideScale = StrideScaleModifierState.ApplyTo(StrideScaleModifier, ActualStrideScale, CachedDeltaTime);
+	check(FMath::IsFinite(ActualStrideScale));
 
 	if (bGraphDrivenWarping)
 	{
@@ -207,6 +224,7 @@ void FAnimNode_StrideWarping::EvaluateSkeletalControl_AnyThread(FComponentSpaceP
 
 		// Now the ScaleOrigin and IKFootLocation are forming a line parallel to the floor, and we can scale the IK foot.
 		const FVector WarpedLocation = ScaleOrigin + (IKFootLocation - ScaleOrigin) * ActualStrideScale;
+		check(!WarpedLocation.ContainsNaN());
 		Foot.IKFootBoneTransform.SetLocation(WarpedLocation);
 
 #if ENABLE_ANIM_DEBUG
