Stridewarping follow up: remove zero check and error case since the code can still function even with a zero stride direction

#Synced-CL 50205727

#Synced-CL 50205727

#Synced-CL 50205727

diff --git a/Engine/Plugins/Animation/AnimationWarping/Source/Runtime/Private/BoneControllers/AnimNode_StrideWarping.cpp b/Engine/Plugins/Animation/AnimationWarping/Source/Runtime/Private/BoneControllers/AnimNode_StrideWarping.cpp
--- a/Engine/Plugins/Animation/AnimationWarping/Source/Runtime/Private/BoneControllers/AnimNode_StrideWarping.cpp
+++ b/Engine/Plugins/Animation/AnimationWarping/Source/Runtime/Private/BoneControllers/AnimNode_StrideWarping.cpp
@@ -88,7 +88,7 @@ void FAnimNode_StrideWarping::EvaluateSkeletalControl_AnyThread(FComponentSpaceP
 	if (!bGraphDrivenWarping)
 	{
 		// If we aren't using graph driven warping, we need to validate the input variables are valid
-		if (StrideDirection.ContainsNaN() || StrideDirection.IsNearlyZero())
+		if (StrideDirection.ContainsNaN())
 		{
 			UE_LOG(LogAnimation, Error, TEXT("Invalid stride direction in FAnimNode_StrideWarping"));
 			return;
