[Animation] Allow ANIM_TRACE_ENABLE to be enabled in test builds without ENABLE_ANIM_DEBUG being true.

#Synced-CL 50654724

diff --git a/Engine/Plugins/Animation/AnimationWarping/Source/Runtime/Private/BoneControllers/AnimNode_OrientationWarping.cpp b/Engine/Plugins/Animation/AnimationWarping/Source/Runtime/Private/BoneControllers/AnimNode_OrientationWarping.cpp
--- a/Engine/Plugins/Animation/AnimationWarping/Source/Runtime/Private/BoneControllers/AnimNode_OrientationWarping.cpp
+++ b/Engine/Plugins/Animation/AnimationWarping/Source/Runtime/Private/BoneControllers/AnimNode_OrientationWarping.cpp
@@ -166,7 +166,7 @@ void FAnimNode_OrientationWarping::EvaluateSkeletalControl_AnyThread(FComponentS
 	bFoundRootMotionAttribute = false;
 #endif
 
-#if ENABLE_ANIM_DEBUG
+#if ENABLE_ANIM_DEBUG || ANIM_TRACE_ENABLED
 	FTransform RootMotionTransformDelta = FTransform::Identity;
 	float RootMotionDeltaAngleRad = 0.0;
 	const float PreviousOrientationAngleRad = ActualOrientationAngleRad;
@@ -200,7 +200,7 @@ void FAnimNode_OrientationWarping::EvaluateSkeletalControl_AnyThread(FComponentS
 
 	if (bGraphDrivenWarping)
 	{
-#if !ENABLE_ANIM_DEBUG
+#if !ENABLE_ANIM_DEBUG && !ANIM_TRACE_ENABLED
 		FTransform RootMotionTransformDelta = FTransform::Identity;
 #endif
 
@@ -318,7 +318,7 @@ void FAnimNode_OrientationWarping::EvaluateSkeletalControl_AnyThread(FComponentS
 				// Don't compensate interpolation by the root motion angle delta if the previous direction isn't valid.
 				if (bCounterCompenstateInterpolationByRootMotion && !PreviousRootMotionDeltaDirection.IsNearlyZero(UE_SMALL_NUMBER))
 				{
-#if !ENABLE_ANIM_DEBUG
+#if !ENABLE_ANIM_DEBUG && !ANIM_TRACE_ENABLED
 					float RootMotionDeltaAngleRad;
 #endif
 					// Counter the interpolated orientation angle by the root motion direction angle delta.
